Material permutation validation system.
Added new plugin (default off) to validate changes that increase material permutations.
A new UMaterialValidationGroup data asset is used to set up validation per group of materials or group of directory paths.
The data asset works as a database which stores all discovered permutations of the materials in it.
Discovery time can be very high for large projects because UMaterialInstanceConstants (and their hard refs) need to be loaded. That is why the results are stored in an asset rather than doing the discovery on demand.
Validation for MICs that add permutation happen live in the material instance editor. Also the validation is run during save and submit validation.
The UMaterialValidationGroup assets used are set up in an extension of the existing DataValidation Editor config panel.
A MaterialValidationCommandlet commandlet is added. It takes the following options:
* -Package= : The package paths to validation group assets to update. If not given we find all validation group assets in the project config.
* -AddFromDirectories= : Whether to add new materials found in the validation group material directories.
* -DumpNonApproved= : Whether to dump material instance asset paths for non approved permutations.
* -UpdateApproved= : Whether to find all material instance permutations and add them as approved.
#rnx

diff --git a/Engine/Plugins/Editor/MaterialValidation/MaterialValidation.uplugin b/Engine/Plugins/Editor/MaterialValidation/MaterialValidation.uplugin
new file mode 100644
--- /dev/null
+++ b/Engine/Plugins/Editor/MaterialValidation/MaterialValidation.uplugin
@@ -0,0 +1,30 @@
+{
+	"FileVersion": 1,
+	"Version": 1,
+	"VersionName": "1.0",
+	"FriendlyName": "Material Validation",
+	"Description": "Additional validation options for materials.",
+	"Category": "Editor",
+	"CreatedBy": "Epic Games, Inc.",
+	"CreatedByURL": "https://epicgames.com",
+	"DocsURL": "",
+	"MarketplaceURL": "",
+	"SupportURL": "",
+	"EnabledByDefault": false,
+	"CanContainContent": false,
+	"IsBetaVersion": true,
+	"Installed": false,
+	"Modules": [
+		{
+			"Name": "MaterialValidation",
+			"Type": "Editor",
+			"LoadingPhase": "Default"
+		}
+	],
+	"Plugins": [
+		{
+			"Name": "DataValidation",
+			"Enabled": true
+		}
+	]
+}
\ No newline at end of file
diff --git a/Engine/Plugins/Editor/MaterialValidation/Source/MaterialValidation.Build.cs b/Engine/Plugins/Editor/MaterialValidation/Source/MaterialValidation.Build.cs
new file mode 100644
--- /dev/null
+++ b/Engine/Plugins/Editor/MaterialValidation/Source/MaterialValidation.Build.cs
@@ -0,0 +1,22 @@
+// Copyright Epic Games, Inc. All Rights Reserved.
+
+namespace UnrealBuildTool.Rules
+{
+	public class MaterialValidation : ModuleRules
+	{
+		public MaterialValidation(ReadOnlyTargetRules Target) : base(Target)
+		{
+			PrivateDependencyModuleNames.AddRange(new string[] {
+				"AssetDefinition",
+				"Core",
+				"CoreUObject",
+				"DataValidation",
+				"DeveloperSettings",
+				"Engine",
+				"MaterialEditor",
+				"RHI",
+				"UnrealEd",
+			});
+		}
+	}
+}
\ No newline at end of file
diff --git a/Engine/Plugins/Editor/MaterialValidation/Source/Private/MaterialValidationCommandlet.cpp b/Engine/Plugins/Editor/MaterialValidation/Source/Private/MaterialValidationCommandlet.cpp
new file mode 100644
--- /dev/null
+++ b/Engine/Plugins/Editor/MaterialValidation/Source/Private/MaterialValidationCommandlet.cpp
@@ -0,0 +1,144 @@
+// Copyright Epic Games, Inc. All Rights Reserved.
+
+#include "MaterialValidationCommandlet.h"
+
+#include "AssetRegistry/AssetData.h"
+#include "AssetRegistry/AssetRegistryModule.h"
+#include "AssetRegistry/IAssetRegistry.h"
+#include "Materials/MaterialInstance.h"
+#include "MaterialValidationGroup.h"
+#include "MaterialValidationLibrary.h"
+#include "MaterialValidationModule.h"
+
+#include UE_INLINE_GENERATED_CPP_BY_NAME(MaterialValidationCommandlet)
+
+UMaterialValidationCommandlet::UMaterialValidationCommandlet(const FObjectInitializer& ObjectInitializer)
+	: Super(ObjectInitializer)
+{
+}
+
+void UMaterialValidationCommandlet::InitializePackageNames(const TArray<FString>& Tokens, TArray<FString>& MapPathNames, bool& bExplicitPackages)
+{
+	// Extract command line switches.
+	for (FString const& Switch : Switches)
+	{
+		int32 Value;
+		if (FParse::Value(*Switch, TEXT("AddFromDirectories="), Value))
+		{
+			bAddMaterialsFromDirectories = Value != 0;
+			UE_LOG(LogMaterialValidation, Log, TEXT("AddFromDirectories: %d."), Value);
+		}
+		else if (FParse::Value(*Switch, TEXT("DumpNonApproved="), Value))
+		{
+			bDumpNonApprovedMaterialInstances = Value != 0;
+			UE_LOG(LogMaterialValidation, Log, TEXT("DumpNonApproved: %d."), Value);
+		}
+		else if (FParse::Value(*Switch, TEXT("UpdateApproved="), Value))
+		{
+			bUpdateApprovedMaterialsInstances = Value != 0;
+			UE_LOG(LogMaterialValidation, Log, TEXT("UpdateApproved: %d."), Value);
+		}
+	}
+
+	// When packages are not explicitly specified, default to the packages that contain the UMaterialValidationGroup objects in the project config.
+	if (!bExplicitPackages)
+	{
+		TArray<UMaterialValidationGroup*> Groups;
+		UMaterialValidationLibrary::GetAllGroups(Groups, /*bInSyncLoad*/true);
+
+		for (UMaterialValidationGroup* Group : Groups)
+		{
+			if (UPackage* Package = Group->GetPackage())
+			{
+				MapPathNames.Add(Package->GetLoadedPath().GetLocalFullPath());
+			
+				UE_LOG(LogMaterialValidation, Log, TEXT("Commandlet Adding Default Package: %s."), *MapPathNames.Last());
+			}
+		}
+
+		bExplicitPackages = true;
+	}
+}
+
+bool UMaterialValidationCommandlet::ShouldSkipPackage(const FString& Filename)
+{
+	if (Super::ShouldSkipPackage(Filename))
+	{
+		return true;
+	}
+
+	// Skip package if we don't find any UMaterialValidationGroup objects in it.
+	FAssetRegistryModule& AssetRegistryModule = FModuleManager::LoadModuleChecked<FAssetRegistryModule>(TEXT("AssetRegistry"));
+	IAssetRegistry& AssetRegistry = AssetRegistryModule.Get();
+
+	FPackagePath PackagePath = FPackagePath::FromLocalPath(Filename);
+	FName PackageName = PackagePath.GetPackageFName();
+
+	TArray<FAssetData> PackageAssetDatas;
+	if (!AssetRegistry.GetAssetsByPackageName(PackageName, PackageAssetDatas, /*bIncludeOnlyOnDiskAssets*/true))
+	{
+		return true;
+	}
+
+	bool bHasValidationGroupAsset = false;
+	for (const FAssetData& AssetData : PackageAssetDatas)
+	{
+		if (AssetData.IsValid() && AssetData.IsInstanceOf(UMaterialValidationGroup::StaticClass()))
+		{
+			bHasValidationGroupAsset = true;
+			break;
+		}
+	}
+
+	return !bHasValidationGroupAsset;
+}
+
+void UMaterialValidationCommandlet::PerformAdditionalOperations(UObject* Object, bool& bSavePackage)
+{
+	UMaterialValidationGroup* Group = Cast<UMaterialValidationGroup>(Object);
+	if (Group == nullptr || !bSavePackage)
+	{
+		return;
+	}
+
+	UE_LOG(LogMaterialValidation, Log, TEXT("Commandlet Begin '%s'"), *Group->GetName());
+
+	if (bAddMaterialsFromDirectories)
+	{
+		UE_LOG(LogMaterialValidation, Log, TEXT("Adding Materials."));
+
+		TArray<FString> Paths;
+		Paths.Reserve(Group->MaterialPaths.Num());
+		for (FDirectoryPath Path : Group->MaterialPaths)
+		{
+			Paths.Add(Path.Path);
+		}
+
+		for (FString const& Path : Paths)
+		{
+			UMaterialValidationLibrary::AddMaterialsDirectoriesToGroup(Group, Paths, /*bInUpdatePermuatations*/false);
+		}
+	}
+
+	if (bDumpNonApprovedMaterialInstances)
+	{
+		UE_LOG(LogMaterialValidation, Log, TEXT("Finding Non Approved Material Instances."));
+
+		TArray<UMaterialInstance*> MaterialInstances;
+		UMaterialValidationLibrary::ValidateAllMaterialsInGroup(Group, MaterialInstances);
+
+		for (UMaterialInstance* MaterialInstance : MaterialInstances)
+		{
+			UE_LOG(LogMaterialValidation, Log, TEXT("Found new material instance '%s' in group '%s'"), *MaterialInstance->GetName(), *Group->GetName());
+		}
+	}
+
+	if (bUpdateApprovedMaterialsInstances)
+	{
+		UE_LOG(LogMaterialValidation, Log, TEXT("Updatting Approved Material Instances."));
+
+		UMaterialValidationLibrary::UpdateAllMaterialsInGroup(Group);
+	}
+
+	UE_LOG(LogMaterialValidation, Log, TEXT("Commandlet End '%s'"), *Group->GetName());
+}
diff --git a/Engine/Plugins/Editor/MaterialValidation/Source/Private/MaterialValidationCommandlet.h b/Engine/Plugins/Editor/MaterialValidation/Source/Private/MaterialValidationCommandlet.h
new file mode 100644
--- /dev/null
+++ b/Engine/Plugins/Editor/MaterialValidation/Source/Private/MaterialValidationCommandlet.h
@@ -0,0 +1,27 @@
+// Copyright Epic Games, Inc. All Rights Reserved.
+
+#pragma once
+
+#include "Commandlets/BaseIteratePackagesCommandlet.h"
+#include "MaterialValidationCommandlet.generated.h"
+
+/** Commandlet to validate and update the UMaterialValidationGroup assets. */
+UCLASS()
+class UMaterialValidationCommandlet : public UBaseIteratePackagesCommandlet
+{
+	GENERATED_UCLASS_BODY()
+
+	//~ Begin UBaseIteratePackagesCommandlet Interface
+	virtual void InitializePackageNames(const TArray<FString>& Tokens, TArray<FString>& MapPathNames, bool& bExplicitPackages) override;
+	virtual bool ShouldSkipPackage(const FString& Filename) override;
+	virtual void PerformAdditionalOperations(class UObject* Object, bool& bSavePackage) override;
+	//~ End UBaseIteratePackagesCommandlet Interface
+
+private:
+	/** true to add any new materials found in the material directories. Enabled with AddFromDirectories=1 on commandline. */
+	bool bAddMaterialsFromDirectories = false;
+	/** true to search for non-approved material instances and dump them to the output. Enabled with DumpNonApproved=1 on commandline. */
+	bool bDumpNonApprovedMaterialInstances = false;
+	/** true to search for all material instances and add them to the approved list. Disabled with UpdateApproved=0 on commandline. */
+	bool bUpdateApprovedMaterialsInstances = true;
+};
diff --git a/Engine/Plugins/Editor/MaterialValidation/Source/Private/MaterialValidationConfig.cpp b/Engine/Plugins/Editor/MaterialValidation/Source/Private/MaterialValidationConfig.cpp
new file mode 100644
--- /dev/null
+++ b/Engine/Plugins/Editor/MaterialValidation/Source/Private/MaterialValidationConfig.cpp
@@ -0,0 +1,32 @@
+// Copyright Epic Games, Inc. All Rights Reserved.
+
+#include "MaterialValidationConfig.h"
+
+#include "DetailCategoryBuilder.h"
+#include "DetailLayoutBuilder.h"
+
+#include UE_INLINE_GENERATED_CPP_BY_NAME(MaterialValidationConfig)
+
+UMaterialValidationConfig::UMaterialValidationConfig(const FObjectInitializer& ObjectInitializer)
+	: Super(ObjectInitializer)
+{
+}
+
+TSharedRef<IDetailCustomization> FDataValidationSettingsCustomization::MakeInstance()
+{
+	return MakeShareable(new FDataValidationSettingsCustomization);
+}
+
+void FDataValidationSettingsCustomization::CustomizeDetails(IDetailLayoutBuilder& DetailBuilder)
+{
+	// Default behavior adds our properties _before_ the base ones, which isn't good for readability.
+	// So break into two categories, and use category order to place our properties last.
+	IDetailCategoryBuilder& BaseCategory = DetailBuilder.EditCategory(TEXT("Data Validation"));
+	BaseCategory.SetSortOrder(0);
+	IDetailCategoryBuilder& ExtendedCategory = DetailBuilder.EditCategory(TEXT("Data Validation Extensions"));
+	ExtendedCategory.SetSortOrder(1);
+
+	UObject* Config = GetMutableDefault<UMaterialValidationConfig>();
+	FAddPropertyParams AddPropertyParams;
+	ExtendedCategory.AddExternalObjects({ Config }, EPropertyLocation::Default, AddPropertyParams.HideRootObjectNode(true));
+}
diff --git a/Engine/Plugins/Editor/MaterialValidation/Source/Private/MaterialValidationConfig.h b/Engine/Plugins/Editor/MaterialValidation/Source/Private/MaterialValidationConfig.h
new file mode 100644
--- /dev/null
+++ b/Engine/Plugins/Editor/MaterialValidation/Source/Private/MaterialValidationConfig.h
@@ -0,0 +1,52 @@
+// Copyright Epic Games, Inc. All Rights Reserved.
+
+#pragma once
+
+#include "Engine/DeveloperSettings.h"
+#include "IDetailCustomization.h"
+
+#include "MaterialValidationConfig.generated.h"
+
+class UMaterialValidationGroup;
+
+/** Enumeration defining when validation happens when it is enabled. */
+UENUM()
+enum class EMaterialValidationLevel
+{
+	DisplayOnly UMETA(ToolTip = "Only show validation result in Material Editor UI."),
+	OnSave UMETA(ToolTip = "Show in Material Editor UI and validate on save."),
+	OnSubmit UMETA(ToolTip = "Show in Material Editor UI and validate on both save and submit."),
+};
+
+/** Configuration for the material validation system. */
+UCLASS(config = Editor, defaultconfig, meta = (DisplayName = "Data Validation Settings"), MinimalAPI)
+class UMaterialValidationConfig : public UDeveloperSettings
+{
+	GENERATED_UCLASS_BODY()
+
+	// Disable auto registration of config UI becuase we use a specific detail customization below.
+	virtual bool SupportsAutoRegistration() const override { return false; }
+
+public:	
+	/** Set true to enable validation. */
+	UPROPERTY(config, EditAnywhere, Category = MaterialPermutationBudget, meta = (DisplayName = "Enable Validation"))
+	bool bEnable = false;
+	/** Set true to enable validation. */
+	UPROPERTY(config, EditAnywhere, Category = MaterialPermutationBudget, meta = (DisplayName = "Enable Save Validation"))
+	EMaterialValidationLevel ValidationLevel = EMaterialValidationLevel::DisplayOnly;
+	/** An array of Material Validation Groups to use during validation. */
+	UPROPERTY(config, EditAnywhere, Category = MaterialPermutationBudget, meta = (EditCondition = "bEnable", DisplayName = "Material Validation Groups"))
+	TArray<TSoftObjectPtr<UMaterialValidationGroup>> Groups;
+};
+
+/** 
+ * Detail customization which adds the UMaterialValidationConfig properties to the core UDataValidationSettings config panel. 
+ * This simplifies the UI by keeping related items in a sinlge place.
+ */
+class FDataValidationSettingsCustomization : public IDetailCustomization
+{
+	virtual void CustomizeDetails(IDetailLayoutBuilder& DetailBuilder) override;
+
+public:
+	static TSharedRef<IDetailCustomization> MakeInstance();
+};
diff --git a/Engine/Plugins/Editor/MaterialValidation/Source/Private/MaterialValidationGroup.cpp b/Engine/Plugins/Editor/MaterialValidation/Source/Private/MaterialValidationGroup.cpp
new file mode 100644
--- /dev/null
+++ b/Engine/Plugins/Editor/MaterialValidation/Source/Private/MaterialValidationGroup.cpp
@@ -0,0 +1,46 @@
+// Copyright Epic Games, Inc. All Rights Reserved.
+
+#include "MaterialValidationGroup.h"
+
+#include "MaterialValidationLibrary.h"
+#include "Misc/Paths.h"
+
+#include UE_INLINE_GENERATED_CPP_BY_NAME(MaterialValidationGroup)
+
+void UMaterialValidationGroup::PostEditChangeProperty(FPropertyChangedEvent& PropertyChangedEvent)
+{
+	Super::PostEditChangeProperty(PropertyChangedEvent);
+
+	if (PropertyChangedEvent.Property && PropertyChangedEvent.MemberProperty->GetFName() == GET_MEMBER_NAME_CHECKED(UMaterialValidationGroup, MaterialPaths))
+	{
+		// Picker returns absolute paths. If possible, convert to /Game/...
+		for (FDirectoryPath& MaterialPath : MaterialPaths)
+		{
+			if (FPaths::IsUnderDirectory(MaterialPath.Path, FPaths::ProjectContentDir()))
+			{
+				FPaths::MakePathRelativeTo(MaterialPath.Path, *FPaths::ProjectContentDir());
+				MaterialPath.Path.InsertAt(0, TEXT("/Game/"));
+			}
+		}
+	}
+}
+
+void UMaterialValidationGroup::UpdateFromPaths()
+{
+	if (MaterialPaths.Num())
+	{
+		TArray<FString> Paths;
+		Paths.Reserve(MaterialPaths.Num());
+		for (FDirectoryPath Path : MaterialPaths)
+		{
+			Paths.Add(Path.Path);
+		}
+	
+		UMaterialValidationLibrary::AddMaterialsDirectoriesToGroup(this, Paths, /*bInUpdatePermuatations*/false);
+	}
+}
+
+void UMaterialValidationGroup::Update()
+{
+	UMaterialValidationLibrary::UpdateAllMaterialsInGroup(this);
+}
diff --git a/Engine/Plugins/Editor/MaterialValidation/Source/Private/MaterialValidationGroup.h b/Engine/Plugins/Editor/MaterialValidation/Source/Private/MaterialValidationGroup.h
new file mode 100644
--- /dev/null
+++ b/Engine/Plugins/Editor/MaterialValidation/Source/Private/MaterialValidationGroup.h
@@ -0,0 +1,56 @@
+// Copyright Epic Games, Inc. All Rights Reserved.
+
+#pragma once
+
+#include "Engine/DataAsset.h"
+#include "MaterialValidationGroup.generated.h"
+
+class UMaterial;
+
+/** Description of validation setup for a single UMaterial. */
+USTRUCT(MinimalAPI)
+struct FMaterialValidationDesc
+{
+	GENERATED_BODY()
+
+	/** Array of currently approved material permutations. */
+	UPROPERTY(VisibleAnywhere, Category = Material)
+	TArray<uint32> ApprovedPermutationHashes;
+};
+
+/** Data Asset containing the validation setup for a group of UMaterials. */
+UCLASS(MinimalAPI)
+class UMaterialValidationGroup : public UDataAsset
+{
+	GENERATED_BODY()
+
+public:
+	/** Array of users that are approved to modify the Material assets in Materials and under the MaterialPaths. */
+	UPROPERTY(EditAnywhere, Category = Material)
+	TArray<FString> Approvers;
+
+	/** Array of directory paths where UMaterial creation requires approval. */
+	UPROPERTY(EditAnywhere, Category = Material)
+	TArray<FDirectoryPath> MaterialPaths;
+
+	/** Map of UMaterial assets to validation descriptions. */
+	UPROPERTY(EditAnywhere, Category = Material, meta = (DisplayThumbnail = False, HidePinAssetPicker = True))
+	TMap<TSoftObjectPtr<UMaterial>, FMaterialValidationDesc> Materials;
+
+	/** 
+	 * Add all Material assets found under the MaterialPaths.
+	 * The new material assets will have empty approved permutation lists.
+	 */
+	UFUNCTION(CallInEditor, Category = Actions, meta = (DisplayName = "Update From Material Paths"))
+	void UpdateFromPaths();
+
+	/** 
+	 * Update current state of all Materials.
+	 * For each material this will traverse the material instance hierachy and add all existing permutations to the approved permutation list.
+	 */
+	UFUNCTION(CallInEditor, Category = Actions, meta = (DisplayName = "Update Approved Lists"))
+	void Update();
+
+protected:
+	virtual void PostEditChangeProperty(FPropertyChangedEvent& PropertyChangedEvent) override;
+};
diff --git a/Engine/Plugins/Editor/MaterialValidation/Source/Private/MaterialValidationLibrary.cpp b/Engine/Plugins/Editor/MaterialValidation/Source/Private/MaterialValidationLibrary.cpp
new file mode 100644
--- /dev/null
+++ b/Engine/Plugins/Editor/MaterialValidation/Source/Private/MaterialValidationLibrary.cpp
@@ -0,0 +1,401 @@
+// Copyright Epic Games, Inc. All Rights Reserved.
+
+#include "MaterialValidationLibrary.h"
+
+#include "AssetRegistry/AssetData.h"
+#include "AssetRegistry/AssetRegistryModule.h"
+#include "AssetRegistry/IAssetRegistry.h"
+#include "Materials/Material.h"
+#include "Materials/MaterialInstanceConstant.h"
+#include "MaterialValidationConfig.h"
+#include "MaterialValidationGroup.h"
+#include "MaterialValidationModule.h"
+#include "Misc/ScopedSlowTask.h"
+
+#include UE_INLINE_GENERATED_CPP_BY_NAME(MaterialValidationLibrary)
+
+#define LOCTEXT_NAMESPACE "MaterialValidation"
+
+namespace MaterialValidationLibrary {
+
+/** 
+ * Build hash for material instance that should be shared by instances that generate the same shader permutations. 
+ * todo: Add non-static switch contributors.
+ */
+static uint32 BuildHash(UMaterialInstanceConstant const& InMaterialInstance, UMaterial const& InBaseMaterial, TConstArrayView<FMaterialParameterInfo> InStaticSwitchParameterInfo)
+{
+	FString KeyString;
+	for (int ParameterIndex = 0; ParameterIndex < InStaticSwitchParameterInfo.Num(); ++ParameterIndex)
+	{
+		FMaterialParameterMetadata InstanceMeta;
+		const bool bInstanceValid = InMaterialInstance.GetParameterValue(EMaterialParameterType::StaticSwitch, InStaticSwitchParameterInfo[ParameterIndex], InstanceMeta);
+		const bool Value = bInstanceValid && InstanceMeta.Value.AsStaticSwitch();
+		
+		FMaterialParameterMetadata BaseMeta;
+		const bool bBaseValid = InBaseMaterial.GetParameterValue(EMaterialParameterType::StaticSwitch, InStaticSwitchParameterInfo[ParameterIndex], BaseMeta);
+		const bool BaseValue = bBaseValid && BaseMeta.Value.AsStaticSwitch();
+
+		if (Value != BaseValue)
+		{
+			InStaticSwitchParameterInfo[ParameterIndex].AppendString(KeyString);
+		}
+	}
+	return GetTypeHash(KeyString);
+}
+
+/** Build hash for material instance that should be shared by instances that generate the same shader permutations. */
+static uint32 BuildHash(UMaterialInstanceConstant const& InMaterialInstance, UMaterial const& InBaseMaterial)
+{
+	TArray<FMaterialParameterInfo> StaticSwitchParameterInfo;
+	TArray<FGuid> StaticSwitchGuids;
+	InBaseMaterial.GetAllStaticSwitchParameterInfo(StaticSwitchParameterInfo, StaticSwitchGuids);
+	return BuildHash(InMaterialInstance, InBaseMaterial, StaticSwitchParameterInfo);
+}
+
+/** 
+ * Load the instance hierachy for a material and add all permutations to the FMaterialValidationDesc. 
+ * This is potentially very slow for a material with many references that may also load/compile textures etc.
+ */
+static void LoadMaterialInstanceConstants(UMaterial const& InMaterial, TArray<UMaterialInstanceConstant*>& OutMaterialInstances)
+{
+	FAssetRegistryModule& AssetRegistryModule = FModuleManager::LoadModuleChecked<FAssetRegistryModule>("AssetRegistry");
+	IAssetRegistry& AssetRegistry = AssetRegistryModule.Get();
+
+	const FAssetData MaterialAssetData = AssetRegistry.GetAssetByObjectPath(InMaterial.GetPathName());
+	if (!MaterialAssetData.IsValid())
+	{
+		return;
+	}
+
+	TQueue<FName> ToVisit;
+	ToVisit.Enqueue(MaterialAssetData.PackageName);
+	TSet<FName> VisitedPackages;
+
+	while (!ToVisit.IsEmpty())
+	{
+		FName CurrentPackageName;
+		ToVisit.Dequeue(CurrentPackageName);
+
+		TArray<FName> ReferencerPackageNames;
+		AssetRegistry.GetReferencers(CurrentPackageName, ReferencerPackageNames);
+
+		for (const FName& PackageName : ReferencerPackageNames)
+		{
+			TArray<FAssetData> AssetDatas;
+			AssetRegistry.GetAssetsByPackageName(PackageName, AssetDatas, true);
+
+			for (FAssetData const& AssetData : AssetDatas)
+			{
+				if (!AssetData.IsValid())
+				{
+					continue;
+				}
+				if (!AssetData.IsInstanceOf(UMaterialInstanceConstant::StaticClass()))
+				{
+					continue;
+				}
+				if (VisitedPackages.Contains(PackageName))
+				{
+					continue;
+				}
+
+				VisitedPackages.Add(PackageName);
+				ToVisit.Enqueue(PackageName);
+				break;
+			}
+		}
+	}
+
+	UE_LOG(LogMaterialValidation, Log, TEXT("Found %d material instances for %s."), VisitedPackages.Num(), *InMaterial.GetName());
+
+	FScopedSlowTask SlowTask(static_cast<float>(VisitedPackages.Num()), LOCTEXT("LoadingMaterialInstances", "Loading Material Instances..."));
+
+	for (FName PackageName : VisitedPackages)
+	{
+		SlowTask.EnterProgressFrame();
+
+		TArray<FAssetData> AssetDatas;
+		AssetRegistry.GetAssetsByPackageName(PackageName, AssetDatas, true);
+		for (FAssetData const& AssetData : AssetDatas)
+		{
+			if (!AssetData.IsValid())
+			{
+				continue;
+			}
+			if (!AssetData.IsInstanceOf(UMaterialInstanceConstant::StaticClass()))
+			{
+				continue;
+			}
+			if (UMaterialInstanceConstant* MaterialInstance = Cast<UMaterialInstanceConstant>(AssetData.GetAsset()))
+			{
+				if (MaterialInstance->GetMaterial() == &InMaterial)
+				{
+					OutMaterialInstances.Add(MaterialInstance);
+				}
+			}
+		}
+	}
+}
+
+/** Load the instance hierachy for a material and add all permutations to the FMaterialValidationDesc. */
+static void UpdateValidMaterialInstances(UMaterial const& InMaterial, FMaterialValidationDesc& OutDesc)
+{
+	TArray<UMaterialInstanceConstant*> MaterialInstances;
+	LoadMaterialInstanceConstants(InMaterial, MaterialInstances);
+
+	TArray<FMaterialParameterInfo> StaticSwitchParameterInfo;
+	TArray<FGuid> StaticSwitchGuids;
+	InMaterial.GetAllStaticSwitchParameterInfo(StaticSwitchParameterInfo, StaticSwitchGuids);
+
+	OutDesc.ApprovedPermutationHashes.Reset();
+	for (UMaterialInstanceConstant* MaterialInstance : MaterialInstances)
+ 	{
+ 		const uint32 Hash = BuildHash(*MaterialInstance, InMaterial, StaticSwitchParameterInfo);
+		if (Hash != 0 && !OutDesc.ApprovedPermutationHashes.Contains(Hash))
+		{
+			UE_LOG(LogMaterialValidation, Log, TEXT("Add material instance '%s' to approved list."), *MaterialInstance->GetName());
+			OutDesc.ApprovedPermutationHashes.Add(Hash);
+		}
+ 	}
+}
+
+/** Load the instance hierachy for a material and find permutations that are not in the FMaterialValidationDesc. */
+static void GetNonValidMaterialInstances(UMaterial const& InMaterial, FMaterialValidationDesc const& InDesc, TArray<UMaterialInstance*>& OutMaterialInstances)
+{
+	TArray<UMaterialInstanceConstant*> MaterialInstances;
+	LoadMaterialInstanceConstants(InMaterial, MaterialInstances);
+
+	TArray<FMaterialParameterInfo> StaticSwitchParameterInfo;
+	TArray<FGuid> StaticSwitchGuids;
+	InMaterial.GetAllStaticSwitchParameterInfo(StaticSwitchParameterInfo, StaticSwitchGuids);
+
+	for (UMaterialInstanceConstant* MaterialInstance : MaterialInstances)
+	{
+		const uint32 Hash = BuildHash(*MaterialInstance, InMaterial, StaticSwitchParameterInfo);
+		if (Hash != 0 && !InDesc.ApprovedPermutationHashes.Contains(Hash))
+		{
+			OutMaterialInstances.Add(MaterialInstance);
+		}
+	}
+}
+
+/** Returns false if the material instance's parent material is in the validation group, but the material instance isn't in the approved list. */
+static bool IsValidationSuccess(UMaterialValidationGroup const& InGroup, UMaterialInstanceConstant const& InMaterialInstance)
+{
+	UMaterial const* Material = InMaterialInstance.GetMaterial();
+	FMaterialValidationDesc const* Settings = InGroup.Materials.Find(Material);
+	if (Settings == nullptr)
+	{
+		return true;
+	}
+
+	uint32 Hash = BuildHash(InMaterialInstance, *Material);
+	return Hash == 0 || Settings->ApprovedPermutationHashes.Contains(Hash);
+}
+
+} //namespace MaterialValidationLibrary
+
+
+UMaterialValidationLibrary::UMaterialValidationLibrary(const FObjectInitializer& ObjectInitializer)
+	: Super(ObjectInitializer)
+{
+}
+
+void UMaterialValidationLibrary::GetAllGroups(TArray<UMaterialValidationGroup*>& OutGroups, bool bInSyncLoad)
+{
+	UMaterialValidationConfig const* Config = GetDefault<UMaterialValidationConfig>();
+	if (!Config->bEnable)
+	{
+		return;
+	}
+
+	for (TSoftObjectPtr<UMaterialValidationGroup> GroupSoftPtr : Config->Groups)
+	{
+		if (bInSyncLoad)
+		{
+			if (UMaterialValidationGroup* Group = GroupSoftPtr.LoadSynchronous())
+			{
+				OutGroups.Add(Group);
+			}
+		}
+		else
+		{
+			if (GroupSoftPtr.IsValid())
+			{
+				if (UMaterialValidationGroup* Group = GroupSoftPtr.Get())
+				{
+					OutGroups.Add(Group);
+				}
+			}
+			else if (!GroupSoftPtr.IsNull())
+			{
+				// Trigger async load in anticipation of being called again.
+				GroupSoftPtr.LoadAsync({});
+			}
+		}
+	}
+}
+
+void UMaterialValidationLibrary::ResetGroup(UMaterialValidationGroup* InGroup)
+{
+	if (InGroup != nullptr)
+	{
+		InGroup->Materials.Reset();
+		InGroup->Approvers.Reset();
+	}
+}
+
+void UMaterialValidationLibrary::AddMaterialsToGroup(UMaterialValidationGroup* InGroup, TArray<UMaterial*> const& InMaterials, bool bInUpdatePermutations)
+{
+	if (InGroup == nullptr)
+	{
+		return;
+	}
+	
+	UE_LOG(LogMaterialValidation, Log, TEXT("Adding %d materials."), InMaterials.Num());
+
+	FScopedSlowTask SlowTask(static_cast<float>(InMaterials.Num()), LOCTEXT("AddingMaterials", "Adding Materials..."), bInUpdatePermutations);
+
+	for (UMaterial* Material : InMaterials)
+	{
+		SlowTask.EnterProgressFrame();
+
+		if (Material != nullptr)
+		{
+			FMaterialValidationDesc& Desc = InGroup->Materials.FindOrAdd(Material);
+			if (bInUpdatePermutations)
+			{
+				MaterialValidationLibrary::UpdateValidMaterialInstances(*Material, Desc);
+			}
+		}
+	}
+
+	InGroup->Modify();
+}
+
+void UMaterialValidationLibrary::AddMaterialsDirectoriesToGroup(UMaterialValidationGroup* InGroup, TArray<FString> const& InDirectoryPaths, bool bInUpdatePermutations)
+{
+	if (InGroup == nullptr)
+	{
+		return;
+	}
+
+	FAssetRegistryModule& AssetRegistryModule = FModuleManager::LoadModuleChecked<FAssetRegistryModule>("AssetRegistry");
+	TArray<FAssetData> AssetDatas;
+
+	for (FString const& DirectoryPath : InDirectoryPaths)
+	{
+		FARFilter Filter;
+		Filter.PackagePaths.Add(*DirectoryPath);
+		Filter.ClassPaths.Add(UMaterial::StaticClass()->GetClassPathName());
+		Filter.bRecursivePaths = true;
+
+		AssetRegistryModule.Get().GetAssets(Filter, AssetDatas);
+	}
+		
+	UE_LOG(LogMaterialValidation, Log, TEXT("Found %d materials in %d directories."), AssetDatas.Num(), InDirectoryPaths.Num());
+
+	FScopedSlowTask SlowTask(static_cast<float>(AssetDatas.Num()), LOCTEXT("LoadingMaterials", "Loading Materials..."), bInUpdatePermutations);
+
+	for (FAssetData const& AssetData : AssetDatas)
+	{
+		SlowTask.EnterProgressFrame();
+
+		if (UMaterial* Material = Cast<UMaterial>(AssetData.GetAsset()))
+		{
+			UE_LOG(LogMaterialValidation, Log, TEXT("Adding material '%s' to '%s'."), *Material->GetName(), *InGroup->GetName());
+
+			FMaterialValidationDesc& Desc = InGroup->Materials.FindOrAdd(Material);
+			if (bInUpdatePermutations)
+			{
+				MaterialValidationLibrary::UpdateValidMaterialInstances(*Material, Desc);
+			}
+		}
+	}
+
+	InGroup->Modify();
+}
+
+void UMaterialValidationLibrary::UpdateAllMaterialsInGroup(UMaterialValidationGroup* InGroup)
+{
+	if (InGroup == nullptr)
+	{
+		return;
+	}
+	
+	UE_LOG(LogMaterialValidation, Log, TEXT("Updating %d materials."), InGroup->Materials.Num());
+
+	FScopedSlowTask SlowTask(static_cast<float>(InGroup->Materials.Num()), LOCTEXT("UpdatingMaterials", "Updating Materials..."));
+	
+	for (TPair<TSoftObjectPtr<UMaterial>, FMaterialValidationDesc>& It : InGroup->Materials)
+	{
+		SlowTask.EnterProgressFrame();
+
+		if (UMaterial const* Material = It.Key.LoadSynchronous())
+		{
+			MaterialValidationLibrary::UpdateValidMaterialInstances(*Material, It.Value);
+		}
+	}
+
+	InGroup->Modify();
+}
+
+void UMaterialValidationLibrary::ValidateAllMaterialsInGroup(UMaterialValidationGroup* InGroup, TArray<UMaterialInstance*>& OutMaterialInstances)
+{
+	if (InGroup == nullptr)
+	{
+		return;
+	}
+
+	UE_LOG(LogMaterialValidation, Log, TEXT("Validating %d materials."), InGroup->Materials.Num());
+
+	FScopedSlowTask SlowTask(static_cast<float>(InGroup->Materials.Num()), LOCTEXT("ValidatingMaterials", "Validating Materials..."));
+	
+	for (TPair<TSoftObjectPtr<UMaterial>, FMaterialValidationDesc>& It : InGroup->Materials)
+	{
+		SlowTask.EnterProgressFrame();
+
+		if (UMaterial const* Material = It.Key.LoadSynchronous())
+		{
+			MaterialValidationLibrary::GetNonValidMaterialInstances(*Material, It.Value, OutMaterialInstances);
+		}
+	}
+}
+
+void UMaterialValidationLibrary::AddMaterialInstanceToGroup(UMaterialValidationGroup* InGroup, UMaterialInstance* InMaterialInstance)
+{
+	if (InGroup != nullptr)
+	{
+		if (UMaterialInstanceConstant* MaterialInstance = Cast<UMaterialInstanceConstant>(InMaterialInstance))
+		{
+			UMaterial* BaseMaterial = MaterialInstance->GetMaterial();
+
+			if (FMaterialValidationDesc* Settings = InGroup->Materials.Find(BaseMaterial))
+			{
+				const uint32 Hash = MaterialValidationLibrary::BuildHash(*MaterialInstance, *BaseMaterial);
+				if (Hash != 0 && !Settings->ApprovedPermutationHashes.Contains(Hash))
+				{
+					UE_LOG(LogMaterialValidation, Log, TEXT("Add material instance '%s' to approved list."), *MaterialInstance->GetName());
+
+					Settings->ApprovedPermutationHashes.Add(Hash);
+					InGroup->Modify();
+				}
+			}
+		}
+	}
+}
+
+bool UMaterialValidationLibrary::ValidateMaterialInstanceInGroup(UMaterialValidationGroup const* InGroup, UMaterialInstance* InMaterialInstance)
+{
+	if (InGroup != nullptr)
+	{
+		if (UMaterialInstanceConstant* MaterialInstance = Cast<UMaterialInstanceConstant>(InMaterialInstance))
+		{
+			return MaterialValidationLibrary::IsValidationSuccess(*InGroup, *MaterialInstance);
+		}
+	}
+	
+	return true;
+}
+
+#undef LOCTEXT_NAMESPACE
diff --git a/Engine/Plugins/Editor/MaterialValidation/Source/Private/MaterialValidationLibrary.h b/Engine/Plugins/Editor/MaterialValidation/Source/Private/MaterialValidationLibrary.h
new file mode 100644
--- /dev/null
+++ b/Engine/Plugins/Editor/MaterialValidation/Source/Private/MaterialValidationLibrary.h
@@ -0,0 +1,92 @@
+// Copyright Epic Games, Inc. All Rights Reserved.
+
+#pragma once
+
+#include "Kismet/BlueprintFunctionLibrary.h"
+#include "MaterialValidationLibrary.generated.h"
+
+class UMaterial;
+class UMaterialInterface;
+class UMaterialInstanceConstant;
+
+/** Material validation functionality exposed as a library of Blueprint functions. */
+UCLASS(MinimalAPI, meta = (ScriptName = "MaterialValidationLibrary"))
+class UMaterialValidationLibrary : public UBlueprintFunctionLibrary
+{
+	GENERATED_UCLASS_BODY()
+
+	/** 
+	 * Get all configured Material Validation Groups.
+	 * 
+	 * @param OutGroups		The returned array of groups.
+	 * @param bInSyncLoad	Whether to sync load the group assets. When false we don't return unloaded group objects, but we do trigger async loading for them.
+	 */
+	UFUNCTION(BlueprintCallable, Category = Material)
+	static void GetAllGroups(TArray<UMaterialValidationGroup*>& OutGroups, bool bInSyncLoad);
+
+	/**
+	 * Remove all UMaterials from a Material Validation Group.
+	 * 
+	 * @param InGroup	The group to reset.
+	 */
+	UFUNCTION(BlueprintCallable, Category = Material)
+	static void ResetGroup(UMaterialValidationGroup* InGroup);
+
+	/** 
+	 * Add UMaterials to a Material Validation Group.
+	 *
+	 * @param InGroup				The group to add materials to.
+	 * @param InMaterials			The array of materials to add.
+	 * @param bInUpdatePermutations	true to calculate and register all existing permutations for the materials (can be slow to load and analyze the tree of material instances).
+	 */
+	UFUNCTION(BlueprintCallable, Category = Material)
+	static void AddMaterialsToGroup(UMaterialValidationGroup* InGroup, TArray<UMaterial*> const& InMaterials, bool bInUpdatePermutations);
+
+	/**
+	 * Add UMaterials found in directories to a Material Validation Group.
+	 *
+	 * @param InGroup				The group to add materials to.
+	 * @param InDirectoryPaths		The array of directories in which to find materials to add.
+	 * @param bInUpdatePermutations	true to calculate and register all existing permutations for the materials (can be slow to load and analyze the tree of material instances).
+	 */
+	UFUNCTION(BlueprintCallable, Category = Material)
+	static void AddMaterialsDirectoriesToGroup(UMaterialValidationGroup* InGroup, TArray<FString> const& InDirectoryPaths, bool bInUpdatePermutations);
+
+	/**
+	 * Recalculate and reregister all existing permutations for all the existing materials in a Material Validation Group (can be slow to load and analyze the tree of material instances).
+	 *
+	 * @param InGroup	The group to update 
+	 */
+	UFUNCTION(BlueprintCallable, Category = Material)
+	static void UpdateAllMaterialsInGroup(UMaterialValidationGroup* InGroup);
+
+	/**
+	 * Find Material Instance assets with permutations that aren't registered with their materials in a Material Validation Group (can be slow to load and analyze the tree of material instances).
+	 *
+	 * @param InGroup				The group to update 
+	 * @param OutMaterialInstances	Returned array of material instances that weren't marked as approved in the group.
+	 */
+	UFUNCTION(BlueprintCallable, Category = Material)
+	static void ValidateAllMaterialsInGroup(UMaterialValidationGroup* InGroup, TArray<UMaterialInstance*>& OutMaterialInstances);
+
+	/**
+	 * Add a material instance to a Material Validation Group.
+	 * It will only be added if it's base material is already in the group.
+	 *
+	 * @param InGroup				The group to update 
+	 * @param	InMaterialInstance	The material instance to add.
+	 */
+	UFUNCTION(BlueprintCallable, Category = Material)
+	static void AddMaterialInstanceToGroup(UMaterialValidationGroup* InGroup, UMaterialInstance* InMaterialInstance);
+
+	/**
+	 * Validate whether a material instance is approved.
+	 * 
+	 * @param	InGroup				The group to update 
+	 * @param	InMaterialInstance	The material instance to validate.
+	 * 
+	 * @returns	false if the material instance belongs to the validation group, but is not in the approved list.
+	 */
+	UFUNCTION(BlueprintCallable, Category = Material)
+	static bool ValidateMaterialInstanceInGroup(UMaterialValidationGroup const* InGroup, UMaterialInstance* InMaterialInstance);
+};
diff --git a/Engine/Plugins/Editor/MaterialValidation/Source/Private/MaterialValidationModule.cpp b/Engine/Plugins/Editor/MaterialValidation/Source/Private/MaterialValidationModule.cpp
new file mode 100644
--- /dev/null
+++ b/Engine/Plugins/Editor/MaterialValidation/Source/Private/MaterialValidationModule.cpp
@@ -0,0 +1,85 @@
+// Copyright Epic Games, Inc. All Rights Reserved.
+
+#include "MaterialValidationModule.h"
+
+#include "Editor.h"
+#include "EditorValidatorSubsystem.h"
+#include "MaterialValidationConfig.h"
+#include "MaterialValidators.h"
+#include "Modules/ModuleManager.h"
+#include "PropertyEditorModule.h"
+
+DEFINE_LOG_CATEGORY(LogMaterialValidation);
+
+void FMaterialValidationModule::StartupModule()
+{
+	if (FModuleManager::Get().IsModuleLoaded("PropertyEditor"))
+	{
+		FPropertyEditorModule& Module = FModuleManager::GetModuleChecked<FPropertyEditorModule>("PropertyEditor");
+		Module.RegisterCustomClassLayout(
+			UDataValidationSettings::StaticClass()->GetFName(),
+			FOnGetDetailCustomizationInstance::CreateStatic(&FDataValidationSettingsCustomization::MakeInstance));
+	}
+
+	// We defer the parts of startup that depend on GEditor until after Editor Init.
+	// Alternatively we could use the PostEngineInit LoadingPhase for this module.
+	// But if we change the LoadingPhase we will need to use 2 modules for this plugin, because the Commandlet must be registered in LoadingPhase Default.
+	FEditorDelegates::OnEditorInitialized.AddRaw(this, &FMaterialValidationModule::OnEditorInitialized);
+	FEditorDelegates::OnEditorPreExit.AddRaw(this, &FMaterialValidationModule::OnEditorPreExit);
+}
+
+void FMaterialValidationModule::OnEditorInitialized(double Duration)
+{
+	if (GEditor != nullptr)
+	{
+		if (UMaterialEditorValidationSubsystem* Subsystem = GEditor->GetEditorSubsystem<UMaterialEditorValidationSubsystem>())
+		{
+			UMaterialEditorValidator_Permutation* Obj = NewObject<UMaterialEditorValidator_Permutation>();
+			Subsystem->AddValidator(Obj);
+			MaterialEditorValidator = Obj;
+		}
+
+		if (UEditorValidatorSubsystem* Subsystem = GEditor->GetEditorSubsystem<UEditorValidatorSubsystem>())
+		{
+			UEditorValidator_MaterialPermutation* Obj = NewObject<UEditorValidator_MaterialPermutation>();
+			Subsystem->AddValidator(Obj);
+			EditorValidator = Obj;
+		}
+	}
+}
+
+void FMaterialValidationModule::OnEditorPreExit()
+{
+	if (GEditor != nullptr)
+	{
+		if (UMaterialEditorValidationSubsystem* Subsystem = GEditor->GetEditorSubsystem<UMaterialEditorValidationSubsystem>())
+		{
+			if (MaterialEditorValidator.IsValid())
+			{
+				Subsystem->RemoveValidator(MaterialEditorValidator.Get());
+			}
+		}
+
+		if (UEditorValidatorSubsystem* Subsystem = GEditor->GetEditorSubsystem<UEditorValidatorSubsystem>())
+		{
+			if (EditorValidator.IsValid())
+			{
+				Subsystem->RemoveValidator(EditorValidator.Get());
+			}
+		}
+	}
+}
+
+void FMaterialValidationModule::ShutdownModule()
+{
+	if (FModuleManager::Get().IsModuleLoaded("PropertyEditor"))
+	{
+		FPropertyEditorModule& Module = FModuleManager::GetModuleChecked<FPropertyEditorModule>("PropertyEditor");
+		Module.UnregisterCustomClassLayout(UDataValidationSettings::StaticClass()->GetFName());
+	}
+
+	FEditorDelegates::OnEditorInitialized.RemoveAll(this);
+	FEditorDelegates::OnEditorPreExit.RemoveAll(this);
+}
+
+IMPLEMENT_MODULE(FMaterialValidationModule, MaterialValidation)
diff --git a/Engine/Plugins/Editor/MaterialValidation/Source/Private/MaterialValidationModule.h b/Engine/Plugins/Editor/MaterialValidation/Source/Private/MaterialValidationModule.h
new file mode 100644
--- /dev/null
+++ b/Engine/Plugins/Editor/MaterialValidation/Source/Private/MaterialValidationModule.h
@@ -0,0 +1,22 @@
+// Copyright Epic Games, Inc. All Rights Reserved.
+
+#pragma once
+
+#include "Logging/LogMacros.h"
+#include "Modules/ModuleInterface.h"
+#include "UObject/WeakObjectPtr.h"
+
+DECLARE_LOG_CATEGORY_EXTERN(LogMaterialValidation, Log, All);
+
+class FMaterialValidationModule : public IModuleInterface
+{
+	virtual void StartupModule() override;
+	virtual void ShutdownModule() override;
+
+private:
+	void OnEditorInitialized(double Duration);
+	void OnEditorPreExit();
+
+	TWeakObjectPtr<class UMaterialEditorValidator_Permutation> MaterialEditorValidator;
+	TWeakObjectPtr<class UEditorValidator_MaterialPermutation> EditorValidator;
+};
diff --git a/Engine/Plugins/Editor/MaterialValidation/Source/Private/MaterialValidators.cpp b/Engine/Plugins/Editor/MaterialValidation/Source/Private/MaterialValidators.cpp
new file mode 100644
--- /dev/null
+++ b/Engine/Plugins/Editor/MaterialValidation/Source/Private/MaterialValidators.cpp
@@ -0,0 +1,117 @@
+// Copyright Epic Games, Inc. All Rights Reserved.
+
+#include "MaterialValidators.h"
+
+#include "AssetRegistry/AssetDataToken.h"
+#include "Logging/TokenizedMessage.h"
+#include "MaterialShared.h"
+#include "Materials/Material.h"
+#include "Materials/MaterialInstanceConstant.h"
+#include "MaterialValidationConfig.h"
+#include "MaterialValidationGroup.h"
+#include "MaterialValidationLibrary.h"
+
+#include UE_INLINE_GENERATED_CPP_BY_NAME(MaterialValidators)
+
+#define LOCTEXT_NAMESPACE "MaterialValidation"
+
+bool UEditorValidator_MaterialPermutation::CanValidateAsset_Implementation(const FAssetData& InAssetData, UObject* InAsset, FDataValidationContext& InContext) const
+{
+	return Cast<UMaterialInstanceConstant>(InAsset) != nullptr;
+}
+
+EDataValidationResult UEditorValidator_MaterialPermutation::ValidateLoadedAsset_Implementation(const FAssetData& InAssetData, UObject* InAsset, FDataValidationContext& InContext)
+{
+	if (UMaterialInstance* MaterialInstance = Cast<UMaterialInstance>(InAsset))
+	{
+		TArray<UMaterialValidationGroup*> Groups;
+		UMaterialValidationLibrary::GetAllGroups(Groups, /*bInSyncLoad*/true);
+
+		for (UMaterialValidationGroup* Group : Groups)
+		{
+			if (!UMaterialValidationLibrary::ValidateMaterialInstanceInGroup(Group, MaterialInstance))
+			{
+				UMaterialValidationConfig const* Config = GetDefault<UMaterialValidationConfig>();
+				const EMessageSeverity::Type MessageSeverity = Config->ValidationLevel == EMaterialValidationLevel::DisplayOnly ? EMessageSeverity::Info : EMessageSeverity::Error;
+
+				TWeakObjectPtr<UMaterialInstance> WeakMaterialInstancePtr = MakeWeakObjectPtr(MaterialInstance);
+				TWeakObjectPtr<UMaterialValidationGroup> WeakGroupPtr = MakeWeakObjectPtr(Group);
+
+				TSharedRef<FTokenizedMessage> Message = AssetMessage(MessageSeverity);
+				Message->AddText(LOCTEXT("EditorValidation_Fail", "The material instance creates an unapproved new material permutation in:"));
+				Message->AddToken(FAssetDataToken::Create(Group));
+				Message->AddText(FText::FromString(TEXT("\n")));
+				Message->AddToken(FActionToken::Create(
+					LOCTEXT("EditorValidation_Fix", "Add to approved list"), FText(),
+					FOnActionTokenExecuted::CreateLambda([WeakMaterialInstancePtr, WeakGroupPtr]()
+					{
+						if (WeakMaterialInstancePtr.IsValid() && WeakGroupPtr.IsValid())
+						{
+							UMaterialValidationLibrary::AddMaterialInstanceToGroup(WeakGroupPtr.Get(), WeakMaterialInstancePtr.Get());
+						}
+					})));
+			}
+		}
+	}
+
+	if (GetValidationResult() != EDataValidationResult::Invalid)
+	{
+		AssetPasses(InAsset);
+	}
+
+	return GetValidationResult(); 
+}
+
+/** 
+ * Get the number of shaders associated with a material instance.
+ * These _may_ be deduplicated with other material instances, so this is an upper
+ * limit estimate of the cost.
+ * todo: Calculate for other platforms, or at least the current preview platform.
+ * todo: Can we get an estimate of memory cost too?
+ */
+static int32 GetShaderCount(UMaterialInstance* InMaterialInstance)
+{
+	TArray<FDebugShaderTypeInfo> ShaderInfos;
+	InMaterialInstance->GetShaderTypes(GMaxRHIShaderPlatform, nullptr, ShaderInfos);
+	
+	int NumShaders = 0;
+	for (FDebugShaderTypeInfo const& ShaderInfo : ShaderInfos)
+	{
+		NumShaders += ShaderInfo.ShaderTypes.Num();
+		for (FDebugShaderPipelineInfo const& PipelineInfo : ShaderInfo.Pipelines)
+		{
+			NumShaders += PipelineInfo.ShaderTypes.Num();
+		}
+	}
+	return NumShaders;
+}
+
+bool UMaterialEditorValidator_Permutation::Validate(UMaterialInterface* InMaterial, FMaterialEditorValidatorContext& InContext) const
+{
+	if (UMaterialInstance* MaterialInstance = Cast<UMaterialInstance>(InMaterial))
+	{
+		TArray<UMaterialValidationGroup*> Groups;
+		UMaterialValidationLibrary::GetAllGroups(Groups, /*bInSyncLoad*/false);
+
+		for (UMaterialValidationGroup const* Group : Groups)
+		{
+			if (!UMaterialValidationLibrary::ValidateMaterialInstanceInGroup(Group, MaterialInstance))
+			{
+				UMaterial const* BaseMaterial = MaterialInstance->GetMaterial();
+				const int32 NumShaders = GetShaderCount(MaterialInstance);
+
+				InContext.Message = FText::Format(
+					LOCTEXT("MaterialEditorValidation_Fail", "This material instance creates an unapproved material permutation of Material '{0}'.\nThis is missing in Material Validation Group '{1}'.\nThis one material permutation is likely to generate {2} unique shaders."),
+					FText::FromString(BaseMaterial->GetName()),
+					FText::FromString(Group->GetName()),
+					FText::AsNumber(NumShaders));
+
+				return false;
+			}
+		}
+	}
+
+	return true;
+}
+
+#undef LOCTEXT_NAMESPACE
diff --git a/Engine/Plugins/Editor/MaterialValidation/Source/Private/MaterialValidators.h b/Engine/Plugins/Editor/MaterialValidation/Source/Private/MaterialValidators.h
new file mode 100644
--- /dev/null
+++ b/Engine/Plugins/Editor/MaterialValidation/Source/Private/MaterialValidators.h
@@ -0,0 +1,27 @@
+// Copyright Epic Games, Inc. All Rights Reserved.
+
+#pragma once
+
+#include "EditorValidatorBase.h"
+#include "MaterialEditorValidation.h"
+
+#include "MaterialValidators.generated.h"
+
+/** Implementation of UEditorValidatorBase that applies validation for limiting the addition of material permutations. */
+UCLASS()
+class UEditorValidator_MaterialPermutation : public UEditorValidatorBase
+{
+	GENERATED_BODY()
+
+	virtual bool CanValidateAsset_Implementation(const FAssetData& InAssetData, UObject* InAsset, FDataValidationContext& InContext) const override;
+	virtual EDataValidationResult ValidateLoadedAsset_Implementation(const FAssetData& InAssetData, UObject* InAsset, FDataValidationContext& InContext) override;
+};
+
+/** Implementation of UMaterialEditorValidatorBase that applies validation for limiting the addition of material permutations. */
+UCLASS()
+class UMaterialEditorValidator_Permutation : public UMaterialEditorValidatorBase
+{
+	GENERATED_BODY()
+
+	virtual bool Validate(UMaterialInterface* InMaterial, FMaterialEditorValidatorContext& InContext) const override;
+};
diff --git a/Engine/Source/Editor/MaterialEditor/MaterialEditor.Build.cs b/Engine/Source/Editor/MaterialEditor/MaterialEditor.Build.cs
--- a/Engine/Source/Editor/MaterialEditor/MaterialEditor.Build.cs
+++ b/Engine/Source/Editor/MaterialEditor/MaterialEditor.Build.cs
@@ -42,10 +42,12 @@ public class MaterialEditor : ModuleRules
                 "Projects",
                 "AssetRegistry",
 				"ToolMenus",
+				"ToolWidgets",
 				"MainFrame",
 				"Landscape",
 				"Kismet", // for rich diffing machinery which was mostly implemented in the kismet module
-				"SourceControl"
+				"SourceControl",
+				"EditorSubsystem"
 			}
 		);
 
diff --git a/Engine/Source/Editor/MaterialEditor/Private/MaterialEditorInstanceDetailCustomization.cpp b/Engine/Source/Editor/MaterialEditor/Private/MaterialEditorInstanceDetailCustomization.cpp
--- a/Engine/Source/Editor/MaterialEditor/Private/MaterialEditorInstanceDetailCustomization.cpp
+++ b/Engine/Source/Editor/MaterialEditor/Private/MaterialEditorInstanceDetailCustomization.cpp
@@ -42,6 +42,9 @@
 #include "Engine/Texture.h"
 #include "HAL/PlatformApplicationMisc.h"
 #include "RenderUtils.h"
+#include "Editor.h"
+#include "MaterialEditorValidation.h"
+#include "SWarningOrErrorBox.h"
 
 #define LOCTEXT_NAMESPACE "MaterialInstanceEditor"
 
@@ -56,6 +59,7 @@ FMaterialInstanceParameterDetails::FMaterialInstanceParameterDetails(UMaterialEd
 	: MaterialEditorInstance(MaterialInstance)
 	, MaterialLayersFunctionsInstance(InMaterialLayersFunctionsInstance) 
 	, ShowHiddenDelegate(InShowHiddenDelegate)
+	, ValidationContext(new FMaterialEditorValidatorContext)
 {
 }
 
@@ -76,6 +80,27 @@ void FMaterialInstanceParameterDetails::CustomizeDetails(IDetailLayoutBuilder& D
 	// Create a new category for a custom layout for the MIC parameters at the very top
 	FName GroupsCategoryName = TEXT("ParameterGroups");
 	IDetailCategoryBuilder& GroupsCategory = DetailLayout.EditCategory(GroupsCategoryName, LOCTEXT("MICParamGroupsTitle", "Parameter Groups"));
+	
+	// Insert optional validation warning row at the top of the category.
+	GroupsCategory.AddCustomRow(LOCTEXT("ValidationWarningRow", "Validation Warning"))
+	.WholeRowContent()
+	[
+		SNew(SWarningOrErrorBox)
+		.Visibility_Lambda([this]() -> EVisibility
+		{
+			bool bWarn = false;
+			if (UMaterialEditorValidationSubsystem* Subsystem = GEditor->GetEditorSubsystem<UMaterialEditorValidationSubsystem>())
+			{
+				bWarn = !Subsystem->Validate(MaterialEditorInstance->GetMaterialInterface(), *ValidationContext.Get());
+			}
+			return bWarn ? EVisibility::Visible : EVisibility::Collapsed;
+		})
+		.Message_Lambda([this]() -> FText
+		{
+			return ValidationContext->Message;
+		})
+	];
+	
 	TSharedRef<IPropertyHandle> ParameterGroupsProperty = DetailLayout.GetProperty("ParameterGroups");
 
 	// check if tree has any selection, we show parameter properties for selected layer only
diff --git a/Engine/Source/Editor/MaterialEditor/Private/MaterialEditorInstanceDetailCustomization.h b/Engine/Source/Editor/MaterialEditor/Private/MaterialEditorInstanceDetailCustomization.h
--- a/Engine/Source/Editor/MaterialEditor/Private/MaterialEditorInstanceDetailCustomization.h
+++ b/Engine/Source/Editor/MaterialEditor/Private/MaterialEditorInstanceDetailCustomization.h
@@ -136,5 +136,8 @@ private:
 
 	/** Associated FMaterialInstance utilities */
 	TWeakPtr<class IPropertyUtilities> PropertyUtilities;
+
+	/** Storage for validation context. */
+	TSharedPtr<struct FMaterialEditorValidatorContext> ValidationContext;
 };
 
diff --git a/Engine/Source/Editor/MaterialEditor/Private/MaterialEditorValidation.cpp b/Engine/Source/Editor/MaterialEditor/Private/MaterialEditorValidation.cpp
new file mode 100644
--- /dev/null
+++ b/Engine/Source/Editor/MaterialEditor/Private/MaterialEditorValidation.cpp
@@ -0,0 +1,30 @@
+// Copyright Epic Games, Inc. All Rights Reserved.
+
+#include "MaterialEditorValidation.h"
+
+#include "Editor.h"
+
+#include UE_INLINE_GENERATED_CPP_BY_NAME(MaterialEditorValidation)
+
+void UMaterialEditorValidationSubsystem::AddValidator(UMaterialEditorValidatorBase* InValidator)
+{
+	Validators.AddUnique(InValidator);
+}
+
+void UMaterialEditorValidationSubsystem::RemoveValidator(UMaterialEditorValidatorBase* InValidator)
+{
+	Validators.Remove(InValidator);
+}
+
+bool UMaterialEditorValidationSubsystem::Validate(UMaterialInterface* InMaterial, FMaterialEditorValidatorContext& InContext) const
+{
+	for (UMaterialEditorValidatorBase* Validator : Validators)
+	{
+		if (Validator != nullptr && !Validator->Validate(InMaterial, InContext))
+		{
+			// Return false on first validation failure.
+			return false;
+		}
+	}
+	return true;
+}
diff --git a/Engine/Source/Editor/MaterialEditor/Public/MaterialEditorValidation.h b/Engine/Source/Editor/MaterialEditor/Public/MaterialEditorValidation.h
new file mode 100644
--- /dev/null
+++ b/Engine/Source/Editor/MaterialEditor/Public/MaterialEditorValidation.h
@@ -0,0 +1,50 @@
+// Copyright Epic Games, Inc. All Rights Reserved.
+
+#pragma once
+
+#include "EditorSubsystem.h"
+#include "MaterialEditorValidation.generated.h"
+
+#define UE_API MATERIALEDITOR_API
+
+class UMaterialInterface;
+
+/** Context for validators to use/fill during Validate() call. */
+struct FMaterialEditorValidatorContext
+{
+	FText Message;
+};
+
+/** Base interface for validators that can be registered with the UMaterialEditorValidationSubsystem. */
+UCLASS(MinimalAPI)
+class UMaterialEditorValidatorBase : public UObject
+{
+	GENERATED_BODY()
+
+public:
+	/** Return false if there are validation errors, and output any validation meesage in InContext. */
+	virtual bool Validate(UMaterialInterface* InMaterial, FMaterialEditorValidatorContext& InContext) const { return true; }
+};
+
+/** Editor subsystem for registering and running material validators in the material editor tools. */
+UCLASS(MinimalAPI, Config = Editor)
+class UMaterialEditorValidationSubsystem : public UEditorSubsystem
+{
+	GENERATED_BODY()
+
+public:
+	/** Register a validator. Added validators will be called on Validate(). */
+	UE_API void AddValidator(UMaterialEditorValidatorBase* InValidator);
+	/** Unregister a validator. */
+	UE_API void RemoveValidator(UMaterialEditorValidatorBase* InValidator);
+	
+	/** Call Validate() on all registered validators.*/
+	bool Validate(UMaterialInterface* InMaterial, FMaterialEditorValidatorContext& InContext) const;
+
+private:
+	/** The array of registered validators. */
+	UPROPERTY()
+	TArray<TObjectPtr<UMaterialEditorValidatorBase>> Validators;
+};
+
+#undef UE_API
