Fix a couple of name/documentation wording issues.
#rnx

diff --git a/Engine/Plugins/Editor/MaterialValidation/Source/Private/MaterialValidationConfig.h b/Engine/Plugins/Editor/MaterialValidation/Source/Private/MaterialValidationConfig.h
--- a/Engine/Plugins/Editor/MaterialValidation/Source/Private/MaterialValidationConfig.h
+++ b/Engine/Plugins/Editor/MaterialValidation/Source/Private/MaterialValidationConfig.h
@@ -9,7 +9,7 @@
 
 class UMaterialValidationGroup;
 
-/** Enumeration defining when validation happens when it is enabled. */
+/** Enumeration defining at what level validation is run. */
 UENUM()
 enum class EMaterialValidationLevel
 {
@@ -32,7 +32,7 @@ public:
 	UPROPERTY(config, EditAnywhere, Category = MaterialPermutationBudget, meta = (DisplayName = "Enable Validation"))
 	bool bEnable = false;
 	/** Set true to enable validation. */
-	UPROPERTY(config, EditAnywhere, Category = MaterialPermutationBudget, meta = (DisplayName = "Enable Save Validation"))
+	UPROPERTY(config, EditAnywhere, Category = MaterialPermutationBudget, meta = (EditCondition = "bEnable"))
 	EMaterialValidationLevel ValidationLevel = EMaterialValidationLevel::DisplayOnly;
 	/** An array of Material Validation Groups to use during validation. */
 	UPROPERTY(config, EditAnywhere, Category = MaterialPermutationBudget, meta = (EditCondition = "bEnable", DisplayName = "Material Validation Groups"))
