MaterialValidation: Support plugin content mount points when setting material directory paths.
#rnx

diff --git a/Engine/Plugins/Editor/MaterialValidation/Source/Private/MaterialValidationGroup.cpp b/Engine/Plugins/Editor/MaterialValidation/Source/Private/MaterialValidationGroup.cpp
--- a/Engine/Plugins/Editor/MaterialValidation/Source/Private/MaterialValidationGroup.cpp
+++ b/Engine/Plugins/Editor/MaterialValidation/Source/Private/MaterialValidationGroup.cpp
@@ -3,6 +3,7 @@
 #include "MaterialValidationGroup.h"
 
 #include "MaterialValidationLibrary.h"
+#include "Misc/PackageName.h"
 #include "Misc/Paths.h"
 
 #include UE_INLINE_GENERATED_CPP_BY_NAME(MaterialValidationGroup)
@@ -11,15 +12,18 @@ void UMaterialValidationGroup::PostEditChangeProperty(FPropertyChangedEvent& Pro
 {
 	Super::PostEditChangeProperty(PropertyChangedEvent);
 
-	if (PropertyChangedEvent.Property && PropertyChangedEvent.MemberProperty->GetFName() == GET_MEMBER_NAME_CHECKED(UMaterialValidationGroup, MaterialPaths))
+	if (PropertyChangedEvent.Property)
 	{
-		// Picker returns absolute paths. If possible, convert to /Game/...
-		for (FDirectoryPath& MaterialPath : MaterialPaths)
+		if (PropertyChangedEvent.MemberProperty && PropertyChangedEvent.MemberProperty->GetFName() == GET_MEMBER_NAME_CHECKED(UMaterialValidationGroup, MaterialPaths))
 		{
-			if (FPaths::IsUnderDirectory(MaterialPath.Path, FPaths::ProjectContentDir()))
+			for (FDirectoryPath& MaterialPath : MaterialPaths)
 			{
-				FPaths::MakePathRelativeTo(MaterialPath.Path, *FPaths::ProjectContentDir());
-				MaterialPath.Path.InsertAt(0, TEXT("/Game/"));
+				// Picker returns absolute paths. Convert to package paths such as /Game/...
+				FString PackagePath;
+				if (FPackageName::TryConvertFilenameToLongPackageName(MaterialPath.Path, PackagePath))
+				{
+					MaterialPath.Path = PackagePath;
+				}
 			}
 		}
 	}
