MaterialValidation: Don't load materials when adding to a group. Instead get the soft object pointer directly from the asset data.
#rnx

diff --git a/Engine/Plugins/Editor/MaterialValidation/Source/Private/MaterialValidationLibrary.cpp b/Engine/Plugins/Editor/MaterialValidation/Source/Private/MaterialValidationLibrary.cpp
--- a/Engine/Plugins/Editor/MaterialValidation/Source/Private/MaterialValidationLibrary.cpp
+++ b/Engine/Plugins/Editor/MaterialValidation/Source/Private/MaterialValidationLibrary.cpp
@@ -295,18 +295,16 @@ void UMaterialValidationLibrary::AddMaterialsDirectoriesToGroup(UMaterialValidat
 		
 	UE_LOG(LogMaterialValidation, Log, TEXT("Found %d materials in %d directories."), AssetDatas.Num(), InDirectoryPaths.Num());
 
-	FScopedSlowTask SlowTask(static_cast<float>(AssetDatas.Num()), LOCTEXT("LoadingMaterials", "Loading Materials..."), bInUpdatePermutations);
-
 	for (FAssetData const& AssetData : AssetDatas)
 	{
-		SlowTask.EnterProgressFrame();
+		UE_LOG(LogMaterialValidation, Log, TEXT("Adding material '%s' to '%s'."), *AssetData.AssetName.ToString(), *InGroup->GetName());
 
-		if (UMaterial* Material = Cast<UMaterial>(AssetData.GetAsset()))
-		{
-			UE_LOG(LogMaterialValidation, Log, TEXT("Adding material '%s' to '%s'."), *Material->GetName(), *InGroup->GetName());
+		TSoftObjectPtr<UMaterial> MaterialSoftPtr(AssetData.ToSoftObjectPath());
+		FMaterialValidationDesc& Desc = InGroup->Materials.FindOrAdd(MaterialSoftPtr);
 
-			FMaterialValidationDesc& Desc = InGroup->Materials.FindOrAdd(Material);
-			if (bInUpdatePermutations)
+		if (bInUpdatePermutations)
+		{
+			if (UMaterial* Material = Cast<UMaterial>(AssetData.GetAsset()))
 			{
 				MaterialValidationLibrary::UpdateValidMaterialInstances(*Material, Desc);
 			}
