Material Validaiton: Add options to auto include project content and project plugin content in a material validation group.
#rnx

diff --git a/Engine/Plugins/Editor/MaterialValidation/Source/MaterialValidation.Build.cs b/Engine/Plugins/Editor/MaterialValidation/Source/MaterialValidation.Build.cs
--- a/Engine/Plugins/Editor/MaterialValidation/Source/MaterialValidation.Build.cs
+++ b/Engine/Plugins/Editor/MaterialValidation/Source/MaterialValidation.Build.cs
@@ -13,6 +13,7 @@ namespace UnrealBuildTool.Rules
 				"DataValidation",
 				"DeveloperSettings",
 				"Engine",
+				"Projects",
 				"MaterialEditor",
 				"RHI",
 				"UnrealEd",
diff --git a/Engine/Plugins/Editor/MaterialValidation/Source/Private/MaterialValidationCommandlet.cpp b/Engine/Plugins/Editor/MaterialValidation/Source/Private/MaterialValidationCommandlet.cpp
--- a/Engine/Plugins/Editor/MaterialValidation/Source/Private/MaterialValidationCommandlet.cpp
+++ b/Engine/Plugins/Editor/MaterialValidation/Source/Private/MaterialValidationCommandlet.cpp
@@ -96,7 +96,7 @@ bool UMaterialValidationCommandlet::ShouldSkipPackage(const FString& Filename)
 void UMaterialValidationCommandlet::PerformAdditionalOperations(UObject* Object, bool& bSavePackage)
 {
 	UMaterialValidationGroup* Group = Cast<UMaterialValidationGroup>(Object);
-	if (Group == nullptr || !bSavePackage)
+	if (Group == nullptr)
 	{
 		return;
 	}
@@ -107,17 +107,7 @@ void UMaterialValidationCommandlet::PerformAdditionalOperations(UObject* Object,
 	{
 		UE_LOG(LogMaterialValidation, Log, TEXT("Adding Materials."));
 
-		TArray<FString> Paths;
-		Paths.Reserve(Group->MaterialPaths.Num());
-		for (FDirectoryPath Path : Group->MaterialPaths)
-		{
-			Paths.Add(Path.Path);
-		}
-
-		for (FString const& Path : Paths)
-		{
-			UMaterialValidationLibrary::AddMaterialsDirectoriesToGroup(Group, Paths, /*bInUpdatePermuatations*/false);
-		}
+		UMaterialValidationLibrary::AddMaterialsFromGroupDirectoryPaths(Group);
 	}
 
 	if (bDumpNonApprovedMaterialInstances)
@@ -140,5 +130,7 @@ void UMaterialValidationCommandlet::PerformAdditionalOperations(UObject* Object,
 		UMaterialValidationLibrary::UpdateAllMaterialsInGroup(Group);
 	}
 
+	bSavePackage = bSavePackage || bAddMaterialsFromDirectories || bUpdateApprovedMaterialsInstances;
+
 	UE_LOG(LogMaterialValidation, Log, TEXT("Commandlet End '%s'"), *Group->GetName());
 }
diff --git a/Engine/Plugins/Editor/MaterialValidation/Source/Private/MaterialValidationConfig.h b/Engine/Plugins/Editor/MaterialValidation/Source/Private/MaterialValidationConfig.h
--- a/Engine/Plugins/Editor/MaterialValidation/Source/Private/MaterialValidationConfig.h
+++ b/Engine/Plugins/Editor/MaterialValidation/Source/Private/MaterialValidationConfig.h
@@ -31,7 +31,7 @@ public:
 	/** Set true to enable validation. */
 	UPROPERTY(config, EditAnywhere, Category = MaterialPermutationBudget, meta = (DisplayName = "Enable Validation"))
 	bool bEnable = false;
-	/** Set true to enable validation. */
+	/** Set levels where validation occurs. */
 	UPROPERTY(config, EditAnywhere, Category = MaterialPermutationBudget, meta = (EditCondition = "bEnable"))
 	EMaterialValidationLevel ValidationLevel = EMaterialValidationLevel::DisplayOnly;
 	/** An array of Material Validation Groups to use during validation. */
diff --git a/Engine/Plugins/Editor/MaterialValidation/Source/Private/MaterialValidationGroup.cpp b/Engine/Plugins/Editor/MaterialValidation/Source/Private/MaterialValidationGroup.cpp
--- a/Engine/Plugins/Editor/MaterialValidation/Source/Private/MaterialValidationGroup.cpp
+++ b/Engine/Plugins/Editor/MaterialValidation/Source/Private/MaterialValidationGroup.cpp
@@ -4,7 +4,6 @@
 
 #include "MaterialValidationLibrary.h"
 #include "Misc/PackageName.h"
-#include "Misc/Paths.h"
 
 #include UE_INLINE_GENERATED_CPP_BY_NAME(MaterialValidationGroup)
 
@@ -31,17 +30,7 @@ void UMaterialValidationGroup::PostEditChangeProperty(FPropertyChangedEvent& Pro
 
 void UMaterialValidationGroup::UpdateFromPaths()
 {
-	if (MaterialPaths.Num())
-	{
-		TArray<FString> Paths;
-		Paths.Reserve(MaterialPaths.Num());
-		for (FDirectoryPath Path : MaterialPaths)
-		{
-			Paths.Add(Path.Path);
-		}
-	
-		UMaterialValidationLibrary::AddMaterialsDirectoriesToGroup(this, Paths, /*bInUpdatePermuatations*/false);
-	}
+	UMaterialValidationLibrary::AddMaterialsFromGroupDirectoryPaths(this);
 }
 
 void UMaterialValidationGroup::Update()
diff --git a/Engine/Plugins/Editor/MaterialValidation/Source/Private/MaterialValidationGroup.h b/Engine/Plugins/Editor/MaterialValidation/Source/Private/MaterialValidationGroup.h
--- a/Engine/Plugins/Editor/MaterialValidation/Source/Private/MaterialValidationGroup.h
+++ b/Engine/Plugins/Editor/MaterialValidation/Source/Private/MaterialValidationGroup.h
@@ -29,10 +29,18 @@ public:
 	UPROPERTY(EditAnywhere, Category = Material)
 	TArray<FString> Approvers;
 
-	/** Array of directory paths where UMaterial creation requires approval. */
+	/** Array of explicit directory paths that we use when searching for UMaterial packages that need validating. */
 	UPROPERTY(EditAnywhere, Category = Material)
 	TArray<FDirectoryPath> MaterialPaths;
 
+	/** Use the project content path in addition to MaterialPaths when searching for UMaterials. */
+	UPROPERTY(EditAnywhere, Category = Material, meta = (DisplayName = "Include All Project Content"))
+	bool bAddProjectPath;
+
+	/** Use all project plugin content paths in addition to MaterialPaths when searching for UMaterials. */
+	UPROPERTY(EditAnywhere, Category = Material, meta = (DisplayName = "Include All Project Plugin Content"))
+	bool bAddProjectPluginPaths;
+
 	/** Map of UMaterial assets to validation descriptions. */
 	UPROPERTY(EditAnywhere, Category = Material, meta = (DisplayThumbnail = False, HidePinAssetPicker = True))
 	TMap<TSoftObjectPtr<UMaterial>, FMaterialValidationDesc> Materials;
diff --git a/Engine/Plugins/Editor/MaterialValidation/Source/Private/MaterialValidationLibrary.cpp b/Engine/Plugins/Editor/MaterialValidation/Source/Private/MaterialValidationLibrary.cpp
--- a/Engine/Plugins/Editor/MaterialValidation/Source/Private/MaterialValidationLibrary.cpp
+++ b/Engine/Plugins/Editor/MaterialValidation/Source/Private/MaterialValidationLibrary.cpp
@@ -5,6 +5,7 @@
 #include "AssetRegistry/AssetData.h"
 #include "AssetRegistry/AssetRegistryModule.h"
 #include "AssetRegistry/IAssetRegistry.h"
+#include "Interfaces/IPluginManager.h"
 #include "Materials/Material.h"
 #include "Materials/MaterialInstanceConstant.h"
 #include "MaterialValidationConfig.h"
@@ -273,45 +274,75 @@ void UMaterialValidationLibrary::AddMaterialsToGroup(UMaterialValidationGroup* I
 	InGroup->Modify();
 }
 
-void UMaterialValidationLibrary::AddMaterialsDirectoriesToGroup(UMaterialValidationGroup* InGroup, TArray<FString> const& InDirectoryPaths, bool bInUpdatePermutations)
+void UMaterialValidationLibrary::AddMaterialsFromGroupDirectoryPaths(UMaterialValidationGroup* InGroup)
 {
 	if (InGroup == nullptr)
 	{
 		return;
 	}
 
+	// Gather all directory search paths.
+	TArray<FString> SearchPaths;
+	for (FDirectoryPath Path : InGroup->MaterialPaths)
+	{
+		SearchPaths.Add(Path.Path);
+	}
+
+	if (InGroup->bAddProjectPath)
+	{
+		SearchPaths.Add(TEXT("/Game/"));
+	}
+
+	if (InGroup->bAddProjectPluginPaths)
+	{
+		IPluginManager& PluginManager = IPluginManager::Get();
+		TArray<TSharedRef<IPlugin>> Plugins = PluginManager.GetDiscoveredPlugins();
+		for (TSharedRef<IPlugin> const& Plugin : Plugins)
+		{
+			if (Plugin->CanContainContent())
+			{
+				if (FPaths::IsUnderDirectory(Plugin->GetContentDir(), FPaths::ProjectDir()))
+				{
+					SearchPaths.Add(FString::Printf(TEXT("/%s/"), *Plugin->GetName()));
+				}
+			}
+		}
+	}
+
+	// Iterate over search paths and gather all UMaterials found in the asset registry.
 	FAssetRegistryModule& AssetRegistryModule = FModuleManager::LoadModuleChecked<FAssetRegistryModule>("AssetRegistry");
 	TArray<FAssetData> AssetDatas;
 
-	for (FString const& DirectoryPath : InDirectoryPaths)
+	for (FString const& SearchPath : SearchPaths)
 	{
 		FARFilter Filter;
-		Filter.PackagePaths.Add(*DirectoryPath);
+		Filter.PackagePaths.Add(*SearchPath);
 		Filter.ClassPaths.Add(UMaterial::StaticClass()->GetClassPathName());
 		Filter.bRecursivePaths = true;
 
 		AssetRegistryModule.Get().GetAssets(Filter, AssetDatas);
 	}
 		
-	UE_LOG(LogMaterialValidation, Log, TEXT("Found %d materials in %d directories."), AssetDatas.Num(), InDirectoryPaths.Num());
+	UE_LOG(LogMaterialValidation, Log, TEXT("Found %d materials in %d directories."), AssetDatas.Num(), SearchPaths.Num());
 
+	// Iterate over UMaterial assets and add them if not found.
+	int32 bNumAdded = 0;
 	for (FAssetData const& AssetData : AssetDatas)
 	{
-		UE_LOG(LogMaterialValidation, Log, TEXT("Adding material '%s' to '%s'."), *AssetData.AssetName.ToString(), *InGroup->GetName());
-
 		TSoftObjectPtr<UMaterial> MaterialSoftPtr(AssetData.ToSoftObjectPath());
-		FMaterialValidationDesc& Desc = InGroup->Materials.FindOrAdd(MaterialSoftPtr);
 
-		if (bInUpdatePermutations)
+		if (!InGroup->Materials.Contains(MaterialSoftPtr))
 		{
-			if (UMaterial* Material = Cast<UMaterial>(AssetData.GetAsset()))
-			{
-				MaterialValidationLibrary::UpdateValidMaterialInstances(*Material, Desc);
-			}
+			UE_LOG(LogMaterialValidation, Log, TEXT("Adding material '%s' to '%s'."), *AssetData.AssetName.ToString(), *InGroup->GetName());
+			InGroup->Materials.Add(MaterialSoftPtr);
+			bNumAdded++;
 		}
 	}
 
-	InGroup->Modify();
+	if (bNumAdded)
+	{
+		InGroup->Modify();
+	}
 }
 
 void UMaterialValidationLibrary::UpdateAllMaterialsInGroup(UMaterialValidationGroup* InGroup)
diff --git a/Engine/Plugins/Editor/MaterialValidation/Source/Private/MaterialValidationLibrary.h b/Engine/Plugins/Editor/MaterialValidation/Source/Private/MaterialValidationLibrary.h
--- a/Engine/Plugins/Editor/MaterialValidation/Source/Private/MaterialValidationLibrary.h
+++ b/Engine/Plugins/Editor/MaterialValidation/Source/Private/MaterialValidationLibrary.h
@@ -43,14 +43,14 @@ class UMaterialValidationLibrary : public UBlueprintFunctionLibrary
 	static void AddMaterialsToGroup(UMaterialValidationGroup* InGroup, TArray<UMaterial*> const& InMaterials, bool bInUpdatePermutations);
 
 	/**
-	 * Add UMaterials found in directories to a Material Validation Group.
+	 * Add all UMaterials to the Material Validation Group which are found in its material search directories.
+	 * This will only add materials, and not remove existing ones.
+	 * Materials are not updated with their permutations, so that newly found materials will have no approved permutations.
 	 *
 	 * @param InGroup				The group to add materials to.
-	 * @param InDirectoryPaths		The array of directories in which to find materials to add.
-	 * @param bInUpdatePermutations	true to calculate and register all existing permutations for the materials (can be slow to load and analyze the tree of material instances).
 	 */
 	UFUNCTION(BlueprintCallable, Category = Material)
-	static void AddMaterialsDirectoriesToGroup(UMaterialValidationGroup* InGroup, TArray<FString> const& InDirectoryPaths, bool bInUpdatePermutations);
+	static void AddMaterialsFromGroupDirectoryPaths(UMaterialValidationGroup* InGroup);
 
 	/**
 	 * Recalculate and reregister all existing permutations for all the existing materials in a Material Validation Group (can be slow to load and analyze the tree of material instances).
