MaterialValidation: Add GetApprovedPermutationCount() BP function.
#rnx

diff --git a/Engine/Plugins/Editor/MaterialValidation/Source/Private/MaterialValidationGroup.h b/Engine/Plugins/Editor/MaterialValidation/Source/Private/MaterialValidationGroup.h
--- a/Engine/Plugins/Editor/MaterialValidation/Source/Private/MaterialValidationGroup.h
+++ b/Engine/Plugins/Editor/MaterialValidation/Source/Private/MaterialValidationGroup.h
@@ -46,7 +46,7 @@ public:
 	TMap<TSoftObjectPtr<UMaterial>, FMaterialValidationDesc> Materials;
 
 	/** 
-	 * Add all Material assets found under the MaterialPaths.
+	 * Add all Material assets found under the material search paths.
 	 * The new material assets will have empty approved permutation lists.
 	 */
 	UFUNCTION(CallInEditor, Category = Actions, meta = (DisplayName = "Update From Material Paths"))
diff --git a/Engine/Plugins/Editor/MaterialValidation/Source/Private/MaterialValidationLibrary.cpp b/Engine/Plugins/Editor/MaterialValidation/Source/Private/MaterialValidationLibrary.cpp
--- a/Engine/Plugins/Editor/MaterialValidation/Source/Private/MaterialValidationLibrary.cpp
+++ b/Engine/Plugins/Editor/MaterialValidation/Source/Private/MaterialValidationLibrary.cpp
@@ -427,4 +427,21 @@ bool UMaterialValidationLibrary::ValidateMaterialInstanceInGroup(UMaterialValida
 	return true;
 }
 
+int32 UMaterialValidationLibrary::GetApprovedPermutationCount(TSoftObjectPtr<UMaterial> InMaterial)
+{
+	TArray<UMaterialValidationGroup*> Groups;
+	GetAllGroups(Groups, true);
+
+	for (UMaterialValidationGroup* Group : Groups)
+	{
+		if (FMaterialValidationDesc* Settings = Group->Materials.Find(InMaterial))
+		{
+			return Settings->ApprovedPermutationHashes.Num();
+		}
+	}
+
+	// The Material was not found.
+	return 0;
+}
+
 #undef LOCTEXT_NAMESPACE
diff --git a/Engine/Plugins/Editor/MaterialValidation/Source/Private/MaterialValidationLibrary.h b/Engine/Plugins/Editor/MaterialValidation/Source/Private/MaterialValidationLibrary.h
--- a/Engine/Plugins/Editor/MaterialValidation/Source/Private/MaterialValidationLibrary.h
+++ b/Engine/Plugins/Editor/MaterialValidation/Source/Private/MaterialValidationLibrary.h
@@ -89,4 +89,16 @@ class UMaterialValidationLibrary : public UBlueprintFunctionLibrary
 	 */
 	UFUNCTION(BlueprintCallable, Category = Material)
 	static bool ValidateMaterialInstanceInGroup(UMaterialValidationGroup const* InGroup, UMaterialInstance* InMaterialInstance);
+
+	/**
+	 * Get the number of approved permutations for a given UMaterial.
+	 * This gets the result directly from the available Material Validation Group data and
+	 * doesn't load or interrogate any material instances.
+	 * 
+	 * @param	InMaterial	The material to get the result for.
+	 * 
+	 * @returns	The number of approved permutations found. This may be less than the number of actual permutations.
+	 */
+	UFUNCTION(BlueprintCallable, Category = Material)
+	static int32 GetApprovedPermutationCount(TSoftObjectPtr<UMaterial> InMaterial);
 };
