MaterialValidation: Change some Log to Display so that they show up more clearly when running the commandlet locally.
#rnx

diff --git a/Engine/Plugins/Editor/MaterialValidation/Source/Private/MaterialValidationCommandlet.cpp b/Engine/Plugins/Editor/MaterialValidation/Source/Private/MaterialValidationCommandlet.cpp
--- a/Engine/Plugins/Editor/MaterialValidation/Source/Private/MaterialValidationCommandlet.cpp
+++ b/Engine/Plugins/Editor/MaterialValidation/Source/Private/MaterialValidationCommandlet.cpp
@@ -101,36 +101,36 @@ void UMaterialValidationCommandlet::PerformAdditionalOperations(UObject* Object,
 		return;
 	}
 
-	UE_LOG(LogMaterialValidation, Log, TEXT("Commandlet Begin '%s'"), *Group->GetName());
+	UE_LOG(LogMaterialValidation, Display, TEXT("Commandlet Begin '%s'"), *Group->GetName());
 
 	if (bAddMaterialsFromDirectories)
 	{
-		UE_LOG(LogMaterialValidation, Log, TEXT("Adding Materials."));
+		UE_LOG(LogMaterialValidation, Display, TEXT("Adding Materials."));
 
 		UMaterialValidationLibrary::AddMaterialsFromGroupDirectoryPaths(Group);
 	}
 
 	if (bDumpNonApprovedMaterialInstances)
 	{
-		UE_LOG(LogMaterialValidation, Log, TEXT("Finding Non Approved Material Instances."));
+		UE_LOG(LogMaterialValidation, Display, TEXT("Finding Non Approved Material Instances."));
 
 		TArray<UMaterialInstance*> MaterialInstances;
 		UMaterialValidationLibrary::ValidateAllMaterialsInGroup(Group, MaterialInstances);
 
 		for (UMaterialInstance* MaterialInstance : MaterialInstances)
 		{
-			UE_LOG(LogMaterialValidation, Log, TEXT("Found new material instance '%s' in group '%s'"), *MaterialInstance->GetName(), *Group->GetName());
+			UE_LOG(LogMaterialValidation, Display, TEXT("Found new material instance '%s' in group '%s'"), *MaterialInstance->GetName(), *Group->GetName());
 		}
 	}
 
 	if (bUpdateApprovedMaterialsInstances)
 	{
-		UE_LOG(LogMaterialValidation, Log, TEXT("Updatting Approved Material Instances."));
+		UE_LOG(LogMaterialValidation, Display, TEXT("Updating Approved Material Instances."));
 
 		UMaterialValidationLibrary::UpdateAllMaterialsInGroup(Group);
 	}
 
 	bSavePackage = bSavePackage || bAddMaterialsFromDirectories || bUpdateApprovedMaterialsInstances;
 
-	UE_LOG(LogMaterialValidation, Log, TEXT("Commandlet End '%s'"), *Group->GetName());
+	UE_LOG(LogMaterialValidation, Display, TEXT("Commandlet End '%s'"), *Group->GetName());
 }
diff --git a/Engine/Plugins/Editor/MaterialValidation/Source/Private/MaterialValidationLibrary.cpp b/Engine/Plugins/Editor/MaterialValidation/Source/Private/MaterialValidationLibrary.cpp
--- a/Engine/Plugins/Editor/MaterialValidation/Source/Private/MaterialValidationLibrary.cpp
+++ b/Engine/Plugins/Editor/MaterialValidation/Source/Private/MaterialValidationLibrary.cpp
@@ -107,7 +107,7 @@ static void LoadMaterialInstanceConstants(UMaterial const& InMaterial, TArray<UM
 		}
 	}
 
-	UE_LOG(LogMaterialValidation, Log, TEXT("Found %d material instances for %s."), VisitedPackages.Num(), *InMaterial.GetName());
+	UE_LOG(LogMaterialValidation, Display, TEXT("Found %d material instances for %s."), VisitedPackages.Num(), *InMaterial.GetName());
 
 	FScopedSlowTask SlowTask(static_cast<float>(VisitedPackages.Num()), LOCTEXT("LoadingMaterialInstances", "Loading Material Instances..."));
 
@@ -253,7 +253,7 @@ void UMaterialValidationLibrary::AddMaterialsToGroup(UMaterialValidationGroup* I
 		return;
 	}
 	
-	UE_LOG(LogMaterialValidation, Log, TEXT("Adding %d materials."), InMaterials.Num());
+	UE_LOG(LogMaterialValidation, Display, TEXT("Adding %d materials."), InMaterials.Num());
 
 	FScopedSlowTask SlowTask(static_cast<float>(InMaterials.Num()), LOCTEXT("AddingMaterials", "Adding Materials..."), bInUpdatePermutations);
 
@@ -323,7 +323,7 @@ void UMaterialValidationLibrary::AddMaterialsFromGroupDirectoryPaths(UMaterialVa
 		AssetRegistryModule.Get().GetAssets(Filter, AssetDatas);
 	}
 		
-	UE_LOG(LogMaterialValidation, Log, TEXT("Found %d materials in %d directories."), AssetDatas.Num(), SearchPaths.Num());
+	UE_LOG(LogMaterialValidation, Display, TEXT("Found %d materials in %d directories."), AssetDatas.Num(), SearchPaths.Num());
 
 	// Iterate over UMaterial assets and add them if not found.
 	int32 bNumAdded = 0;
@@ -352,7 +352,7 @@ void UMaterialValidationLibrary::UpdateAllMaterialsInGroup(UMaterialValidationGr
 		return;
 	}
 	
-	UE_LOG(LogMaterialValidation, Log, TEXT("Updating %d materials."), InGroup->Materials.Num());
+	UE_LOG(LogMaterialValidation, Display, TEXT("Updating %d materials."), InGroup->Materials.Num());
 
 	FScopedSlowTask SlowTask(static_cast<float>(InGroup->Materials.Num()), LOCTEXT("UpdatingMaterials", "Updating Materials..."));
 	
@@ -376,7 +376,7 @@ void UMaterialValidationLibrary::ValidateAllMaterialsInGroup(UMaterialValidation
 		return;
 	}
 
-	UE_LOG(LogMaterialValidation, Log, TEXT("Validating %d materials."), InGroup->Materials.Num());
+	UE_LOG(LogMaterialValidation, Display, TEXT("Validating %d materials."), InGroup->Materials.Num());
 
 	FScopedSlowTask SlowTask(static_cast<float>(InGroup->Materials.Num()), LOCTEXT("ValidatingMaterials", "Validating Materials..."));
 	
