MaterialValidataion: Remove developer content from validation file search.
Add option to remove materials that are no longer on disk.
#rnx

diff --git a/Engine/Plugins/Editor/MaterialValidation/Source/Private/MaterialValidationCommandlet.cpp b/Engine/Plugins/Editor/MaterialValidation/Source/Private/MaterialValidationCommandlet.cpp
--- a/Engine/Plugins/Editor/MaterialValidation/Source/Private/MaterialValidationCommandlet.cpp
+++ b/Engine/Plugins/Editor/MaterialValidation/Source/Private/MaterialValidationCommandlet.cpp
@@ -23,6 +23,11 @@ void UMaterialValidationCommandlet::InitializePackageNames(const TArray<FString>
 	for (FString const& Switch : Switches)
 	{
 		int32 Value;
+		if (FParse::Value(*Switch, TEXT("RemoveMissing="), Value))
+		{
+			bRemoveMissingMaterials = Value != 0;
+			UE_LOG(LogMaterialValidation, Log, TEXT("RemoveMissing: %d."), Value);
+		}
 		if (FParse::Value(*Switch, TEXT("AddFromDirectories="), Value))
 		{
 			bAddMaterialsFromDirectories = Value != 0;
@@ -103,6 +108,13 @@ void UMaterialValidationCommandlet::PerformAdditionalOperations(UObject* Object,
 
 	UE_LOG(LogMaterialValidation, Display, TEXT("Commandlet Begin '%s'"), *Group->GetName());
 
+	if (bRemoveMissingMaterials)
+	{
+		UE_LOG(LogMaterialValidation, Display, TEXT("Removing Materials."));
+
+		UMaterialValidationLibrary::RemoveMissingMaterialsFromGroup(Group);
+	}
+
 	if (bAddMaterialsFromDirectories)
 	{
 		UE_LOG(LogMaterialValidation, Display, TEXT("Adding Materials."));
@@ -130,7 +142,7 @@ void UMaterialValidationCommandlet::PerformAdditionalOperations(UObject* Object,
 		UMaterialValidationLibrary::UpdateAllMaterialsInGroup(Group);
 	}
 
-	bSavePackage = bSavePackage || bAddMaterialsFromDirectories || bUpdateApprovedMaterialsInstances;
+	bSavePackage = bSavePackage || bRemoveMissingMaterials || bAddMaterialsFromDirectories || bUpdateApprovedMaterialsInstances;
 
 	UE_LOG(LogMaterialValidation, Display, TEXT("Commandlet End '%s'"), *Group->GetName());
 }
diff --git a/Engine/Plugins/Editor/MaterialValidation/Source/Private/MaterialValidationCommandlet.h b/Engine/Plugins/Editor/MaterialValidation/Source/Private/MaterialValidationCommandlet.h
--- a/Engine/Plugins/Editor/MaterialValidation/Source/Private/MaterialValidationCommandlet.h
+++ b/Engine/Plugins/Editor/MaterialValidation/Source/Private/MaterialValidationCommandlet.h
@@ -18,8 +18,10 @@ class UMaterialValidationCommandlet : public UBaseIteratePackagesCommandlet
 	//~ End UBaseIteratePackagesCommandlet Interface
 
 private:
-	/** true to add any new materials found in the material directories. Enabled with AddFromDirectories=1 on commandline. */
-	bool bAddMaterialsFromDirectories = false;
+	/** true to remove any materials that are not found on disk. Disabled with RemoveMissing=0 on commandline. */
+	bool bRemoveMissingMaterials = true;
+	/** true to add any new materials found in the material directories. Disabled with AddFromDirectories=0 on commandline. */
+	bool bAddMaterialsFromDirectories = true;
 	/** true to search for non-approved material instances and dump them to the output. Enabled with DumpNonApproved=1 on commandline. */
 	bool bDumpNonApprovedMaterialInstances = false;
 	/** true to search for all material instances and add them to the approved list. Disabled with UpdateApproved=0 on commandline. */
diff --git a/Engine/Plugins/Editor/MaterialValidation/Source/Private/MaterialValidationGroup.cpp b/Engine/Plugins/Editor/MaterialValidation/Source/Private/MaterialValidationGroup.cpp
--- a/Engine/Plugins/Editor/MaterialValidation/Source/Private/MaterialValidationGroup.cpp
+++ b/Engine/Plugins/Editor/MaterialValidation/Source/Private/MaterialValidationGroup.cpp
@@ -28,6 +28,11 @@ void UMaterialValidationGroup::PostEditChangeProperty(FPropertyChangedEvent& Pro
 	}
 }
 
+void UMaterialValidationGroup::RemoveMissing()
+{
+	UMaterialValidationLibrary::RemoveMissingMaterialsFromGroup(this);
+}
+
 void UMaterialValidationGroup::UpdateFromPaths()
 {
 	UMaterialValidationLibrary::AddMaterialsFromGroupDirectoryPaths(this);
diff --git a/Engine/Plugins/Editor/MaterialValidation/Source/Private/MaterialValidationGroup.h b/Engine/Plugins/Editor/MaterialValidation/Source/Private/MaterialValidationGroup.h
--- a/Engine/Plugins/Editor/MaterialValidation/Source/Private/MaterialValidationGroup.h
+++ b/Engine/Plugins/Editor/MaterialValidation/Source/Private/MaterialValidationGroup.h
@@ -45,6 +45,10 @@ public:
 	UPROPERTY(EditAnywhere, Category = Material, meta = (DisplayThumbnail = False, HidePinAssetPicker = True))
 	TMap<TSoftObjectPtr<UMaterial>, FMaterialValidationDesc> Materials;
 
+	/** Remove all Material assets that cannot be found on disk. */
+	UFUNCTION(CallInEditor, Category = Actions, meta = (DisplayName = "Remove Missing Materials"))
+	void RemoveMissing();
+
 	/** 
 	 * Add all Material assets found under the material search paths.
 	 * The new material assets will have empty approved permutation lists.
diff --git a/Engine/Plugins/Editor/MaterialValidation/Source/Private/MaterialValidationLibrary.cpp b/Engine/Plugins/Editor/MaterialValidation/Source/Private/MaterialValidationLibrary.cpp
--- a/Engine/Plugins/Editor/MaterialValidation/Source/Private/MaterialValidationLibrary.cpp
+++ b/Engine/Plugins/Editor/MaterialValidation/Source/Private/MaterialValidationLibrary.cpp
@@ -329,6 +329,12 @@ void UMaterialValidationLibrary::AddMaterialsFromGroupDirectoryPaths(UMaterialVa
 	int32 bNumAdded = 0;
 	for (FAssetData const& AssetData : AssetDatas)
 	{
+		//todo: Expose UEditorValidatorSubsystem::ExcludedDirectories and skip them all.
+		if (AssetData.PackagePath.ToString().StartsWith(TEXT("/Game/Developers/")))
+		{
+			continue;
+		}
+
 		TSoftObjectPtr<UMaterial> MaterialSoftPtr(AssetData.ToSoftObjectPath());
 
 		if (!InGroup->Materials.Contains(MaterialSoftPtr))
@@ -345,6 +351,37 @@ void UMaterialValidationLibrary::AddMaterialsFromGroupDirectoryPaths(UMaterialVa
 	}
 }
 
+void UMaterialValidationLibrary::RemoveMissingMaterialsFromGroup(UMaterialValidationGroup* InGroup)
+{
+	if (InGroup == nullptr)
+	{
+		return;
+	}
+
+	TArray<TSoftObjectPtr<UMaterial>> ToRemove;
+	for (TPair<TSoftObjectPtr<UMaterial>, FMaterialValidationDesc>& It : InGroup->Materials)
+	{
+		FString PackageName = It.Key.ToSoftObjectPath().GetLongPackageName();
+		if (!FPackageName::DoesPackageExist(PackageName))
+		{
+			ToRemove.Add(It.Key);
+		}
+	}
+
+	UE_LOG(LogMaterialValidation, Display, TEXT("Removing %d materials."), ToRemove.Num());
+
+	for (TSoftObjectPtr<UMaterial> const& Material : ToRemove)
+	{
+		UE_LOG(LogMaterialValidation, Log, TEXT("Removing material '%s' from '%s'."), *Material.GetAssetName(), *InGroup->GetName());
+		InGroup->Materials.Remove(Material);
+	}
+
+	if (ToRemove.Num() > 0)
+	{
+		InGroup->Modify();
+	}
+}
+
 void UMaterialValidationLibrary::UpdateAllMaterialsInGroup(UMaterialValidationGroup* InGroup)
 {
 	if (InGroup == nullptr)
diff --git a/Engine/Plugins/Editor/MaterialValidation/Source/Private/MaterialValidationLibrary.h b/Engine/Plugins/Editor/MaterialValidation/Source/Private/MaterialValidationLibrary.h
--- a/Engine/Plugins/Editor/MaterialValidation/Source/Private/MaterialValidationLibrary.h
+++ b/Engine/Plugins/Editor/MaterialValidation/Source/Private/MaterialValidationLibrary.h
@@ -52,6 +52,14 @@ class UMaterialValidationLibrary : public UBlueprintFunctionLibrary
 	UFUNCTION(BlueprintCallable, Category = Material)
 	static void AddMaterialsFromGroupDirectoryPaths(UMaterialValidationGroup* InGroup);
 
+	/**
+	 * Remove all UMaterials that no longer exist on disk from the Material Validation Group.
+	 *
+	 * @param InGroup	The group to update
+	 */
+	UFUNCTION(BlueprintCallable, Category = Material)
+	static void RemoveMissingMaterialsFromGroup(UMaterialValidationGroup* InGroup);
+
 	/**
 	 * Recalculate and reregister all existing permutations for all the existing materials in a Material Validation Group (can be slow to load and analyze the tree of material instances).
 	 *
