MaterialValidation: Remove action to check out database file to avoid causing contention issues for now.
#rnx

diff --git a/Engine/Plugins/Editor/MaterialValidation/Source/Private/MaterialValidators.cpp b/Engine/Plugins/Editor/MaterialValidation/Source/Private/MaterialValidators.cpp
--- a/Engine/Plugins/Editor/MaterialValidation/Source/Private/MaterialValidators.cpp
+++ b/Engine/Plugins/Editor/MaterialValidation/Source/Private/MaterialValidators.cpp
@@ -40,16 +40,20 @@ EDataValidationResult UEditorValidator_MaterialPermutation::ValidateLoadedAsset_
 				TSharedRef<FTokenizedMessage> Message = AssetMessage(MessageSeverity);
 				Message->AddText(LOCTEXT("EditorValidation_Fail", "The material instance creates an unapproved new material permutation in:"));
 				Message->AddToken(FAssetDataToken::Create(Group));
-				Message->AddText(FText::FromString(TEXT("\n")));
-				Message->AddToken(FActionToken::Create(
-					LOCTEXT("EditorValidation_Fix", "Add to approved list"), FText(),
-					FOnActionTokenExecuted::CreateLambda([WeakMaterialInstancePtr, WeakGroupPtr]()
-					{
-						if (WeakMaterialInstancePtr.IsValid() && WeakGroupPtr.IsValid())
+
+				if (MessageSeverity == EMessageSeverity::Error)
+				{
+					Message->AddText(FText::FromString(TEXT("\n")));
+					Message->AddToken(FActionToken::Create(
+						LOCTEXT("EditorValidation_Fix", "Add to approved list"), FText(),
+						FOnActionTokenExecuted::CreateLambda([WeakMaterialInstancePtr, WeakGroupPtr]()
 						{
-							UMaterialValidationLibrary::AddMaterialInstanceToGroup(WeakGroupPtr.Get(), WeakMaterialInstancePtr.Get());
-						}
-					})));
+							if (WeakMaterialInstancePtr.IsValid() && WeakGroupPtr.IsValid())
+							{
+								UMaterialValidationLibrary::AddMaterialInstanceToGroup(WeakGroupPtr.Get(), WeakMaterialInstancePtr.Get());
+							}
+						})));
+				}
 			}
 		}
 	}
