Material validation: Add validator for material saves.
Initially the validator only fails if this is a new material is in the Material Validation Group search path but hasn't been seen before.
Also changed current warn messages for validators to be less scary while the submit validation is still not used.
#rnx

diff --git a/Engine/Plugins/Editor/MaterialValidation/Source/Private/MaterialValidationLibrary.cpp b/Engine/Plugins/Editor/MaterialValidation/Source/Private/MaterialValidationLibrary.cpp
--- a/Engine/Plugins/Editor/MaterialValidation/Source/Private/MaterialValidationLibrary.cpp
+++ b/Engine/Plugins/Editor/MaterialValidation/Source/Private/MaterialValidationLibrary.cpp
@@ -53,6 +53,82 @@ static uint32 BuildHash(UMaterialInstanceConstant const& InMaterialInstance, UMa
 	return BuildHash(InMaterialInstance, InBaseMaterial, StaticSwitchParameterInfo);
 }
 
+/** Gather all search and exclude directory paths for a group. */
+static void GetSearchPaths(UMaterialValidationGroup const& InGroup, TArray<FString>& OutPaths, TArray<FString>& OutExcludePaths)
+{
+	for (FDirectoryPath Path : InGroup.MaterialPaths)
+	{
+		OutPaths.Add(Path.Path);
+	}
+
+	if (InGroup.bAddProjectPath)
+	{
+		OutPaths.Add(TEXT("/Game/"));
+
+		//todo: Expose UEditorValidatorSubsystem::ExcludedDirectories and skip them all.
+		OutExcludePaths.Add(TEXT("/Game/Developers/"));
+	}
+
+	if (InGroup.bAddProjectPluginPaths)
+	{
+		IPluginManager& PluginManager = IPluginManager::Get();
+		TArray<TSharedRef<IPlugin>> Plugins = PluginManager.GetDiscoveredPlugins();
+		for (TSharedRef<IPlugin> const& Plugin : Plugins)
+		{
+			if (Plugin->CanContainContent())
+			{
+				if (FPaths::IsUnderDirectory(Plugin->GetContentDir(), FPaths::ProjectDir()))
+				{
+					OutPaths.Add(FString::Printf(TEXT("/%s/"), *Plugin->GetName()));
+				}
+			}
+		}
+	}
+}
+
+/** Gather all assets of a given class using the given search and exclued paths. */
+static void GetAssetsFromSearchPaths(UClass* InClass, TConstArrayView<FString> InSearchPaths, TConstArrayView<FString> InExcludePaths, TArray<FAssetData>& OutAssetDatas)
+{
+	FAssetRegistryModule& AssetRegistryModule = FModuleManager::LoadModuleChecked<FAssetRegistryModule>("AssetRegistry");
+
+	for (FString const& SearchPath : InSearchPaths)
+	{
+		FARFilter Filter;
+		Filter.PackagePaths.Add(*SearchPath);
+		Filter.ClassPaths.Add(InClass->GetClassPathName());
+		Filter.bRecursivePaths = true;
+
+		const int32 BaseAddedIndex = OutAssetDatas.Num();
+		AssetRegistryModule.Get().GetAssets(Filter, OutAssetDatas);
+
+		// FARFilter doesn't support exclude paths. So we have to manually prune :(
+		TArray<FString> RelevantExcludePaths;
+		for (FString const& ExcludePath : InExcludePaths)
+		{
+			if (ExcludePath.StartsWith(SearchPath))
+			{
+				RelevantExcludePaths.Add(ExcludePath);
+			}
+		}
+		if (RelevantExcludePaths.Num() > 0)
+		{
+			for (int32 Index = OutAssetDatas.Num() - 1; Index >= BaseAddedIndex; --Index)
+			{
+				FAssetData const& AssetData = OutAssetDatas[Index];
+				FString PackagePath = AssetData.PackagePath.ToString();
+				for (FString const& ExcludePath : RelevantExcludePaths)
+				{
+					if (PackagePath.StartsWith(ExcludePath))
+					{
+						OutAssetDatas.RemoveAtSwap(Index, 1, EAllowShrinking::No);
+						break;
+					}
+				}
+			}
+		}
+	}
+}
+
 /** 
  * Load the instance hierachy for a material and add all permutations to the FMaterialValidationDesc. 
  * This is potentially very slow for a material with many references that may also load/compile textures etc.
@@ -281,47 +357,13 @@ void UMaterialValidationLibrary::AddMaterialsFromGroupDirectoryPaths(UMaterialVa
 		return;
 	}
 
-	// Gather all directory search paths.
+	// Iterate over search paths and gather all UMaterials found in the asset registry.
 	TArray<FString> SearchPaths;
-	for (FDirectoryPath Path : InGroup->MaterialPaths)
-	{
-		SearchPaths.Add(Path.Path);
-	}
-
-	if (InGroup->bAddProjectPath)
-	{
-		SearchPaths.Add(TEXT("/Game/"));
-	}
+	TArray<FString> ExcludePaths;
+	MaterialValidationLibrary::GetSearchPaths(*InGroup, SearchPaths, ExcludePaths);
 
-	if (InGroup->bAddProjectPluginPaths)
-	{
-		IPluginManager& PluginManager = IPluginManager::Get();
-		TArray<TSharedRef<IPlugin>> Plugins = PluginManager.GetDiscoveredPlugins();
-		for (TSharedRef<IPlugin> const& Plugin : Plugins)
-		{
-			if (Plugin->CanContainContent())
-			{
-				if (FPaths::IsUnderDirectory(Plugin->GetContentDir(), FPaths::ProjectDir()))
-				{
-					SearchPaths.Add(FString::Printf(TEXT("/%s/"), *Plugin->GetName()));
-				}
-			}
-		}
-	}
-
-	// Iterate over search paths and gather all UMaterials found in the asset registry.
-	FAssetRegistryModule& AssetRegistryModule = FModuleManager::LoadModuleChecked<FAssetRegistryModule>("AssetRegistry");
 	TArray<FAssetData> AssetDatas;
-
-	for (FString const& SearchPath : SearchPaths)
-	{
-		FARFilter Filter;
-		Filter.PackagePaths.Add(*SearchPath);
-		Filter.ClassPaths.Add(UMaterial::StaticClass()->GetClassPathName());
-		Filter.bRecursivePaths = true;
-
-		AssetRegistryModule.Get().GetAssets(Filter, AssetDatas);
-	}
+	MaterialValidationLibrary::GetAssetsFromSearchPaths(UMaterial::StaticClass(), SearchPaths, ExcludePaths, AssetDatas);
 		
 	UE_LOG(LogMaterialValidation, Display, TEXT("Found %d materials in %d directories."), AssetDatas.Num(), SearchPaths.Num());
 
@@ -329,12 +371,6 @@ void UMaterialValidationLibrary::AddMaterialsFromGroupDirectoryPaths(UMaterialVa
 	int32 bNumAdded = 0;
 	for (FAssetData const& AssetData : AssetDatas)
 	{
-		//todo: Expose UEditorValidatorSubsystem::ExcludedDirectories and skip them all.
-		if (AssetData.PackagePath.ToString().StartsWith(TEXT("/Game/Developers/")))
-		{
-			continue;
-		}
-
 		TSoftObjectPtr<UMaterial> MaterialSoftPtr(AssetData.ToSoftObjectPath());
 
 		if (!InGroup->Materials.Contains(MaterialSoftPtr))
@@ -464,6 +500,43 @@ bool UMaterialValidationLibrary::ValidateMaterialInstanceInGroup(UMaterialValida
 	return true;
 }
 
+bool UMaterialValidationLibrary::ValidateMaterialInGroup(UMaterialValidationGroup const* InGroup, UMaterial* InMaterial)
+{
+	if (InGroup == nullptr || InMaterial == nullptr)
+	{
+		return true;
+	}
+
+	if (InGroup->Materials.Contains(InMaterial))
+	{
+		return true;
+	}
+
+	FString PackagePath = InMaterial->GetPathName();
+
+	TArray<FString> SearchPaths;
+	TArray<FString> ExcludePaths;
+	MaterialValidationLibrary::GetSearchPaths(*InGroup, SearchPaths, ExcludePaths);
+
+	for (FString const& Path : ExcludePaths)
+	{
+		if (PackagePath.StartsWith(Path))
+		{
+			return true;
+		}
+	}
+
+	for (FString const& Path : SearchPaths)
+	{
+		if (PackagePath.StartsWith(Path))
+		{
+			return false;
+		}
+	}
+
+	return true;
+}
+
 int32 UMaterialValidationLibrary::GetApprovedPermutationCount(TSoftObjectPtr<UMaterial> InMaterial)
 {
 	TArray<UMaterialValidationGroup*> Groups;
diff --git a/Engine/Plugins/Editor/MaterialValidation/Source/Private/MaterialValidationLibrary.h b/Engine/Plugins/Editor/MaterialValidation/Source/Private/MaterialValidationLibrary.h
--- a/Engine/Plugins/Editor/MaterialValidation/Source/Private/MaterialValidationLibrary.h
+++ b/Engine/Plugins/Editor/MaterialValidation/Source/Private/MaterialValidationLibrary.h
@@ -63,7 +63,7 @@ class UMaterialValidationLibrary : public UBlueprintFunctionLibrary
 	/**
 	 * Recalculate and reregister all existing permutations for all the existing materials in a Material Validation Group (can be slow to load and analyze the tree of material instances).
 	 *
-	 * @param InGroup	The group to update 
+	 * @param InGroup	The group to update.
 	 */
 	UFUNCTION(BlueprintCallable, Category = Material)
 	static void UpdateAllMaterialsInGroup(UMaterialValidationGroup* InGroup);
@@ -71,7 +71,7 @@ class UMaterialValidationLibrary : public UBlueprintFunctionLibrary
 	/**
 	 * Find Material Instance assets with permutations that aren't registered with their materials in a Material Validation Group (can be slow to load and analyze the tree of material instances).
 	 *
-	 * @param InGroup				The group to update 
+	 * @param InGroup				The group to update.
 	 * @param OutMaterialInstances	Returned array of material instances that weren't marked as approved in the group.
 	 */
 	UFUNCTION(BlueprintCallable, Category = Material)
@@ -81,7 +81,7 @@ class UMaterialValidationLibrary : public UBlueprintFunctionLibrary
 	 * Add a material instance to a Material Validation Group.
 	 * It will only be added if it's base material is already in the group.
 	 *
-	 * @param InGroup				The group to update 
+	 * @param InGroup				The group to update.
 	 * @param	InMaterialInstance	The material instance to add.
 	 */
 	UFUNCTION(BlueprintCallable, Category = Material)
@@ -90,7 +90,7 @@ class UMaterialValidationLibrary : public UBlueprintFunctionLibrary
 	/**
 	 * Validate whether a material instance is approved.
 	 * 
-	 * @param	InGroup				The group to update 
+	 * @param	InGroup				The group to use.
 	 * @param	InMaterialInstance	The material instance to validate.
 	 * 
 	 * @returns	false if the material instance belongs to the validation group, but is not in the approved list.
@@ -98,6 +98,17 @@ class UMaterialValidationLibrary : public UBlueprintFunctionLibrary
 	UFUNCTION(BlueprintCallable, Category = Material)
 	static bool ValidateMaterialInstanceInGroup(UMaterialValidationGroup const* InGroup, UMaterialInstance* InMaterialInstance);
 
+	/**
+	 * Validate whether a material is approved.
+	 * 
+	 * @param	InGroup		The group to use.
+	 * @param	InMaterial	The material to validate.
+	 * 
+	 * @returns	false if the material is in the validation group search path, but is not in the material list.
+	 */
+	UFUNCTION(BlueprintCallable, Category = Material)
+	static bool ValidateMaterialInGroup(UMaterialValidationGroup const* InGroup, UMaterial* InMaterial);
+
 	/**
 	 * Get the number of approved permutations for a given UMaterial.
 	 * This gets the result directly from the available Material Validation Group data and
diff --git a/Engine/Plugins/Editor/MaterialValidation/Source/Private/MaterialValidationModule.cpp b/Engine/Plugins/Editor/MaterialValidation/Source/Private/MaterialValidationModule.cpp
--- a/Engine/Plugins/Editor/MaterialValidation/Source/Private/MaterialValidationModule.cpp
+++ b/Engine/Plugins/Editor/MaterialValidation/Source/Private/MaterialValidationModule.cpp
@@ -34,16 +34,20 @@ void FMaterialValidationModule::OnEditorInitialized(double Duration)
 	{
 		if (UMaterialEditorValidationSubsystem* Subsystem = GEditor->GetEditorSubsystem<UMaterialEditorValidationSubsystem>())
 		{
-			UMaterialEditorValidator_Permutation* Obj = NewObject<UMaterialEditorValidator_Permutation>();
-			Subsystem->AddValidator(Obj);
-			MaterialEditorValidator = Obj;
+			UMaterialEditorValidatorBase* Validator = NewObject<UMaterialEditorValidator_Permutation>();
+			Subsystem->AddValidator(Validator);
+			MaterialEditorValidators.Add(Validator);
 		}
 
 		if (UEditorValidatorSubsystem* Subsystem = GEditor->GetEditorSubsystem<UEditorValidatorSubsystem>())
 		{
-			UEditorValidator_MaterialPermutation* Obj = NewObject<UEditorValidator_MaterialPermutation>();
-			Subsystem->AddValidator(Obj);
-			EditorValidator = Obj;
+			UEditorValidatorBase* ValidatorMaterialInstance = NewObject<UEditorValidator_MaterialInstancePermutation>();
+			Subsystem->AddValidator(ValidatorMaterialInstance);
+			EditorValidators.Add(ValidatorMaterialInstance);
+
+			UEditorValidatorBase* ValidatorMaterial = NewObject<UEditorValidator_MaterialPermutation>();
+			Subsystem->AddValidator(ValidatorMaterial);
+			EditorValidators.Add(ValidatorMaterial);
 		}
 	}
 }
@@ -54,17 +58,23 @@ void FMaterialValidationModule::OnEditorPreExit()
 	{
 		if (UMaterialEditorValidationSubsystem* Subsystem = GEditor->GetEditorSubsystem<UMaterialEditorValidationSubsystem>())
 		{
-			if (MaterialEditorValidator.IsValid())
+			for (TWeakObjectPtr<UMaterialEditorValidatorBase>& Validator : MaterialEditorValidators)
 			{
-				Subsystem->RemoveValidator(MaterialEditorValidator.Get());
+				if (Validator.IsValid())
+				{
+					Subsystem->RemoveValidator(Validator.Get());
+				}
 			}
 		}
 
 		if (UEditorValidatorSubsystem* Subsystem = GEditor->GetEditorSubsystem<UEditorValidatorSubsystem>())
 		{
-			if (EditorValidator.IsValid())
+			for (TWeakObjectPtr<UEditorValidatorBase>& Validator : EditorValidators)
 			{
-				Subsystem->RemoveValidator(EditorValidator.Get());
+				if (Validator.IsValid())
+				{
+					Subsystem->RemoveValidator(Validator.Get());
+				}
 			}
 		}
 	}
diff --git a/Engine/Plugins/Editor/MaterialValidation/Source/Private/MaterialValidationModule.h b/Engine/Plugins/Editor/MaterialValidation/Source/Private/MaterialValidationModule.h
--- a/Engine/Plugins/Editor/MaterialValidation/Source/Private/MaterialValidationModule.h
+++ b/Engine/Plugins/Editor/MaterialValidation/Source/Private/MaterialValidationModule.h
@@ -17,6 +17,6 @@ private:
 	void OnEditorInitialized(double Duration);
 	void OnEditorPreExit();
 
-	TWeakObjectPtr<class UMaterialEditorValidator_Permutation> MaterialEditorValidator;
-	TWeakObjectPtr<class UEditorValidator_MaterialPermutation> EditorValidator;
+	TArray<TWeakObjectPtr<class UMaterialEditorValidatorBase>> MaterialEditorValidators;
+	TArray<TWeakObjectPtr<class UEditorValidatorBase>> EditorValidators;
 };
diff --git a/Engine/Plugins/Editor/MaterialValidation/Source/Private/MaterialValidators.cpp b/Engine/Plugins/Editor/MaterialValidation/Source/Private/MaterialValidators.cpp
--- a/Engine/Plugins/Editor/MaterialValidation/Source/Private/MaterialValidators.cpp
+++ b/Engine/Plugins/Editor/MaterialValidation/Source/Private/MaterialValidators.cpp
@@ -17,10 +17,47 @@
 
 bool UEditorValidator_MaterialPermutation::CanValidateAsset_Implementation(const FAssetData& InAssetData, UObject* InAsset, FDataValidationContext& InContext) const
 {
-	return Cast<UMaterialInstanceConstant>(InAsset) != nullptr;
+	return Cast<UMaterial>(InAsset) != nullptr;
 }
 
 EDataValidationResult UEditorValidator_MaterialPermutation::ValidateLoadedAsset_Implementation(const FAssetData& InAssetData, UObject* InAsset, FDataValidationContext& InContext)
+{
+	if (UMaterial* Material = Cast<UMaterial>(InAsset))
+	{
+		TArray<UMaterialValidationGroup*> Groups;
+		UMaterialValidationLibrary::GetAllGroups(Groups, /*bInSyncLoad*/true);
+
+		for (UMaterialValidationGroup* Group : Groups)
+		{
+			if (!UMaterialValidationLibrary::ValidateMaterialInGroup(Group, Material))
+			{
+				UMaterialValidationConfig const* Config = GetDefault<UMaterialValidationConfig>();
+				const EMessageSeverity::Type MessageSeverity = Config->ValidationLevel == EMaterialValidationLevel::DisplayOnly ? EMessageSeverity::Info : EMessageSeverity::Error;
+
+				TWeakObjectPtr<UMaterial> WeakMaterialPtr = MakeWeakObjectPtr(Material);
+				TWeakObjectPtr<UMaterialValidationGroup> WeakGroupPtr = MakeWeakObjectPtr(Group);
+
+				TSharedRef<FTokenizedMessage> Message = AssetMessage(MessageSeverity);
+				Message->AddText(LOCTEXT("EditorValidation_Fail", "This is an unapproved new material in:"));
+				Message->AddToken(FAssetDataToken::Create(Group));
+			}
+		}
+	}
+
+	if (GetValidationResult() != EDataValidationResult::Invalid)
+	{
+		AssetPasses(InAsset);
+	}
+
+	return GetValidationResult();
+}
+
+bool UEditorValidator_MaterialInstancePermutation::CanValidateAsset_Implementation(const FAssetData& InAssetData, UObject* InAsset, FDataValidationContext& InContext) const
+{
+	return Cast<UMaterialInstanceConstant>(InAsset) != nullptr;
+}
+
+EDataValidationResult UEditorValidator_MaterialInstancePermutation::ValidateLoadedAsset_Implementation(const FAssetData& InAssetData, UObject* InAsset, FDataValidationContext& InContext)
 {
 	if (UMaterialInstance* MaterialInstance = Cast<UMaterialInstance>(InAsset))
 	{
@@ -38,22 +75,8 @@ EDataValidationResult UEditorValidator_MaterialPermutation::ValidateLoadedAsset_
 				TWeakObjectPtr<UMaterialValidationGroup> WeakGroupPtr = MakeWeakObjectPtr(Group);
 
 				TSharedRef<FTokenizedMessage> Message = AssetMessage(MessageSeverity);
-				Message->AddText(LOCTEXT("EditorValidation_Fail", "The material instance creates an unapproved new material permutation in:"));
+				Message->AddText(LOCTEXT("EditorValidation_Fail", "The material instance creates a new material permutation in:"));
 				Message->AddToken(FAssetDataToken::Create(Group));
-
-				if (MessageSeverity == EMessageSeverity::Error)
-				{
-					Message->AddText(FText::FromString(TEXT("\n")));
-					Message->AddToken(FActionToken::Create(
-						LOCTEXT("EditorValidation_Fix", "Add to approved list"), FText(),
-						FOnActionTokenExecuted::CreateLambda([WeakMaterialInstancePtr, WeakGroupPtr]()
-						{
-							if (WeakMaterialInstancePtr.IsValid() && WeakGroupPtr.IsValid())
-							{
-								UMaterialValidationLibrary::AddMaterialInstanceToGroup(WeakGroupPtr.Get(), WeakMaterialInstancePtr.Get());
-							}
-						})));
-				}
 			}
 		}
 	}
@@ -105,9 +128,8 @@ bool UMaterialEditorValidator_Permutation::Validate(UMaterialInterface* InMateri
 				const int32 NumShaders = GetShaderCount(MaterialInstance);
 
 				InContext.Message = FText::Format(
-					LOCTEXT("MaterialEditorValidation_Fail", "This material instance creates an unapproved material permutation of Material '{0}'.\nThis is missing in Material Validation Group '{1}'.\nThis one material permutation is likely to generate {2} unique shaders."),
+					LOCTEXT("MaterialEditorValidation_Warn", "This material instance creates a new permutation of Base Material '{0}'.\nA permutation of this material is likely to generate {1} unique shaders.\nCan you use an existing permutation instead?"),
 					FText::FromString(BaseMaterial->GetName()),
-					FText::FromString(Group->GetName()),
 					FText::AsNumber(NumShaders));
 
 				return false;
diff --git a/Engine/Plugins/Editor/MaterialValidation/Source/Private/MaterialValidators.h b/Engine/Plugins/Editor/MaterialValidation/Source/Private/MaterialValidators.h
--- a/Engine/Plugins/Editor/MaterialValidation/Source/Private/MaterialValidators.h
+++ b/Engine/Plugins/Editor/MaterialValidation/Source/Private/MaterialValidators.h
@@ -7,7 +7,7 @@
 
 #include "MaterialValidators.generated.h"
 
-/** Implementation of UEditorValidatorBase that applies validation for limiting the addition of material permutations. */
+/** Implementation of UEditorValidatorBase that applies validation for limiting the addition of material permutations from UMaterials. */
 UCLASS()
 class UEditorValidator_MaterialPermutation : public UEditorValidatorBase
 {
@@ -17,6 +17,16 @@ class UEditorValidator_MaterialPermutation : public UEditorValidatorBase
 	virtual EDataValidationResult ValidateLoadedAsset_Implementation(const FAssetData& InAssetData, UObject* InAsset, FDataValidationContext& InContext) override;
 };
 
+/** Implementation of UEditorValidatorBase that applies validation for limiting the addition of material permutations from UMaterialInstances. */
+UCLASS()
+class UEditorValidator_MaterialInstancePermutation : public UEditorValidatorBase
+{
+	GENERATED_BODY()
+
+	virtual bool CanValidateAsset_Implementation(const FAssetData& InAssetData, UObject* InAsset, FDataValidationContext& InContext) const override;
+	virtual EDataValidationResult ValidateLoadedAsset_Implementation(const FAssetData& InAssetData, UObject* InAsset, FDataValidationContext& InContext) override;
+};
+
 /** Implementation of UMaterialEditorValidatorBase that applies validation for limiting the addition of material permutations. */
 UCLASS()
 class UMaterialEditorValidator_Permutation : public UMaterialEditorValidatorBase
