Material Validation: Fix duplicate text localization key.
#rnx

diff --git a/Engine/Plugins/Editor/MaterialValidation/Source/Private/MaterialValidators.cpp b/Engine/Plugins/Editor/MaterialValidation/Source/Private/MaterialValidators.cpp
--- a/Engine/Plugins/Editor/MaterialValidation/Source/Private/MaterialValidators.cpp
+++ b/Engine/Plugins/Editor/MaterialValidation/Source/Private/MaterialValidators.cpp
@@ -38,7 +38,7 @@ EDataValidationResult UEditorValidator_MaterialPermutation::ValidateLoadedAsset_
 				TWeakObjectPtr<UMaterialValidationGroup> WeakGroupPtr = MakeWeakObjectPtr(Group);
 
 				TSharedRef<FTokenizedMessage> Message = AssetMessage(MessageSeverity);
-				Message->AddText(LOCTEXT("EditorValidation_Fail", "This is an unapproved new material in:"));
+				Message->AddText(LOCTEXT("Fail_NewMaterial", "This is an unapproved new material in:"));
 				Message->AddToken(FAssetDataToken::Create(Group));
 			}
 		}
@@ -75,7 +75,7 @@ EDataValidationResult UEditorValidator_MaterialInstancePermutation::ValidateLoad
 				TWeakObjectPtr<UMaterialValidationGroup> WeakGroupPtr = MakeWeakObjectPtr(Group);
 
 				TSharedRef<FTokenizedMessage> Message = AssetMessage(MessageSeverity);
-				Message->AddText(LOCTEXT("EditorValidation_Fail", "The material instance creates a new material permutation in:"));
+				Message->AddText(LOCTEXT("Fail_NewPermutation", "The material instance creates a new material permutation in:"));
 				Message->AddToken(FAssetDataToken::Create(Group));
 			}
 		}
