Add exclude paths to material valiation group.
#rnx

diff --git a/Engine/Plugins/Editor/MaterialValidation/Source/Private/MaterialValidationGroup.cpp b/Engine/Plugins/Editor/MaterialValidation/Source/Private/MaterialValidationGroup.cpp
--- a/Engine/Plugins/Editor/MaterialValidation/Source/Private/MaterialValidationGroup.cpp
+++ b/Engine/Plugins/Editor/MaterialValidation/Source/Private/MaterialValidationGroup.cpp
@@ -7,23 +7,32 @@
 
 #include UE_INLINE_GENERATED_CPP_BY_NAME(MaterialValidationGroup)
 
+// Path picker returns absolute paths. Convert to package paths such as /Game/...
+static void ConvertPaths(TArrayView<FDirectoryPath> Paths)
+{
+	for (FDirectoryPath& Path : Paths)
+	{
+		FString PackagePath;
+		if (FPackageName::TryConvertFilenameToLongPackageName(Path.Path, PackagePath))
+		{
+			Path.Path = PackagePath;
+		}
+	}
+}
+
 void UMaterialValidationGroup::PostEditChangeProperty(FPropertyChangedEvent& PropertyChangedEvent)
 {
 	Super::PostEditChangeProperty(PropertyChangedEvent);
 
-	if (PropertyChangedEvent.Property)
+	if (PropertyChangedEvent.Property && PropertyChangedEvent.MemberProperty)
 	{
-		if (PropertyChangedEvent.MemberProperty && PropertyChangedEvent.MemberProperty->GetFName() == GET_MEMBER_NAME_CHECKED(UMaterialValidationGroup, MaterialPaths))
+		if (PropertyChangedEvent.MemberProperty->GetFName() == GET_MEMBER_NAME_CHECKED(UMaterialValidationGroup, MaterialPaths))
+		{
+			ConvertPaths(MaterialPaths);
+		}
+		else if (PropertyChangedEvent.MemberProperty->GetFName() == GET_MEMBER_NAME_CHECKED(UMaterialValidationGroup, MaterialExcludePaths))
 		{
-			for (FDirectoryPath& MaterialPath : MaterialPaths)
-			{
-				// Picker returns absolute paths. Convert to package paths such as /Game/...
-				FString PackagePath;
-				if (FPackageName::TryConvertFilenameToLongPackageName(MaterialPath.Path, PackagePath))
-				{
-					MaterialPath.Path = PackagePath;
-				}
-			}
+			ConvertPaths(MaterialExcludePaths);
 		}
 	}
 }
diff --git a/Engine/Plugins/Editor/MaterialValidation/Source/Private/MaterialValidationGroup.h b/Engine/Plugins/Editor/MaterialValidation/Source/Private/MaterialValidationGroup.h
--- a/Engine/Plugins/Editor/MaterialValidation/Source/Private/MaterialValidationGroup.h
+++ b/Engine/Plugins/Editor/MaterialValidation/Source/Private/MaterialValidationGroup.h
@@ -33,6 +33,10 @@ public:
 	UPROPERTY(EditAnywhere, Category = Material)
 	TArray<FDirectoryPath> MaterialPaths;
 
+	/** Array of explicit directory paths that we exclude when searching for UMaterial packages that need validating. */
+	UPROPERTY(EditAnywhere, Category = Material)
+	TArray<FDirectoryPath> MaterialExcludePaths;
+
 	/** Use the project content path in addition to MaterialPaths when searching for UMaterials. */
 	UPROPERTY(EditAnywhere, Category = Material, meta = (DisplayName = "Include All Project Content"))
 	bool bAddProjectPath;
@@ -58,7 +62,7 @@ public:
 
 	/** 
 	 * Update current state of all Materials.
-	 * For each material this will traverse the material instance hierachy and add all existing permutations to the approved permutation list.
+	 * For each material this will traverse the material instance hierarchy and add all existing permutations to the approved permutation list.
 	 */
 	UFUNCTION(CallInEditor, Category = Actions, meta = (DisplayName = "Update Approved Lists"))
 	void Update();
diff --git a/Engine/Plugins/Editor/MaterialValidation/Source/Private/MaterialValidationLibrary.cpp b/Engine/Plugins/Editor/MaterialValidation/Source/Private/MaterialValidationLibrary.cpp
--- a/Engine/Plugins/Editor/MaterialValidation/Source/Private/MaterialValidationLibrary.cpp
+++ b/Engine/Plugins/Editor/MaterialValidation/Source/Private/MaterialValidationLibrary.cpp
@@ -61,6 +61,11 @@ static void GetSearchPaths(UMaterialValidationGroup const& InGroup, TArray<FStri
 		OutPaths.Add(Path.Path);
 	}
 
+	for (FDirectoryPath Path : InGroup.MaterialExcludePaths)
+	{
+		OutExcludePaths.Add(Path.Path);
+	}
+
 	if (InGroup.bAddProjectPath)
 	{
 		OutPaths.Add(TEXT("/Game/"));
@@ -86,6 +91,19 @@ static void GetSearchPaths(UMaterialValidationGroup const& InGroup, TArray<FStri
 	}
 }
 
+/** Returns true if the PackagePath is found under any of the paths in the InPaths array. */
+static bool IsPackageInPaths(FString const& PackagePath, TConstArrayView<FString> InPaths)
+{
+	for (FString const& Path : InPaths)
+	{
+		if (PackagePath.StartsWith(Path))
+		{
+			return true;
+		}
+	}
+	return false;
+}
+
 /** Gather all assets of a given class using the given search and exclued paths. */
 static void GetAssetsFromSearchPaths(UClass* InClass, TConstArrayView<FString> InSearchPaths, TConstArrayView<FString> InExcludePaths, TArray<FAssetData>& OutAssetDatas)
 {
@@ -116,13 +134,9 @@ static void GetAssetsFromSearchPaths(UClass* InClass, TConstArrayView<FString> I
 			{
 				FAssetData const& AssetData = OutAssetDatas[Index];
 				FString PackagePath = AssetData.PackagePath.ToString();
-				for (FString const& ExcludePath : RelevantExcludePaths)
+				if (IsPackageInPaths(PackagePath, RelevantExcludePaths))
 				{
-					if (PackagePath.StartsWith(ExcludePath))
-					{
-						OutAssetDatas.RemoveAtSwap(Index, 1, EAllowShrinking::No);
-						break;
-					}
+					OutAssetDatas.RemoveAtSwap(Index, 1, EAllowShrinking::No);
 				}
 			}
 		}
@@ -130,7 +144,7 @@ static void GetAssetsFromSearchPaths(UClass* InClass, TConstArrayView<FString> I
 }
 
 /** 
- * Load the instance hierachy for a material and add all permutations to the FMaterialValidationDesc. 
+ * Load the instance hierarchy for a material and add all permutations to the FMaterialValidationDesc. 
  * This is potentially very slow for a material with many references that may also load/compile textures etc.
  */
 static void LoadMaterialInstanceConstants(UMaterial const& InMaterial, TArray<UMaterialInstanceConstant*>& OutMaterialInstances)
@@ -214,7 +228,7 @@ static void LoadMaterialInstanceConstants(UMaterial const& InMaterial, TArray<UM
 	}
 }
 
-/** Load the instance hierachy for a material and add all permutations to the FMaterialValidationDesc. */
+/** Load the instance hierarchy for a material and add all permutations to the FMaterialValidationDesc. */
 static void UpdateValidMaterialInstances(UMaterial const& InMaterial, FMaterialValidationDesc& OutDesc)
 {
 	TArray<UMaterialInstanceConstant*> MaterialInstances;
@@ -236,7 +250,7 @@ static void UpdateValidMaterialInstances(UMaterial const& InMaterial, FMaterialV
  	}
 }
 
-/** Load the instance hierachy for a material and find permutations that are not in the FMaterialValidationDesc. */
+/** Load the instance hierarchy for a material and find permutations that are not in the FMaterialValidationDesc. */
 static void GetNonValidMaterialInstances(UMaterial const& InMaterial, FMaterialValidationDesc const& InDesc, TArray<UMaterialInstance*>& OutMaterialInstances)
 {
 	TArray<UMaterialInstanceConstant*> MaterialInstances;
@@ -394,11 +408,27 @@ void UMaterialValidationLibrary::RemoveMissingMaterialsFromGroup(UMaterialValida
 		return;
 	}
 
+	TArray<FString> SearchPaths;
+	TArray<FString> ExcludePaths;
+	MaterialValidationLibrary::GetSearchPaths(*InGroup, SearchPaths, ExcludePaths);
+
 	TArray<TSoftObjectPtr<UMaterial>> ToRemove;
 	for (TPair<TSoftObjectPtr<UMaterial>, FMaterialValidationDesc>& It : InGroup->Materials)
 	{
 		FString PackageName = It.Key.ToSoftObjectPath().GetLongPackageName();
 		if (!FPackageName::DoesPackageExist(PackageName))
+		{
+			ToRemove.Add(It.Key);
+			continue;
+		}
+
+		if (MaterialValidationLibrary::IsPackageInPaths(PackageName, ExcludePaths))
+		{
+			ToRemove.Add(It.Key);
+			continue;
+		}
+
+		if (!MaterialValidationLibrary::IsPackageInPaths(PackageName, SearchPaths))
 		{
 			ToRemove.Add(It.Key);
 		}
@@ -518,20 +548,14 @@ bool UMaterialValidationLibrary::ValidateMaterialInGroup(UMaterialValidationGrou
 	TArray<FString> ExcludePaths;
 	MaterialValidationLibrary::GetSearchPaths(*InGroup, SearchPaths, ExcludePaths);
 
-	for (FString const& Path : ExcludePaths)
+	if (MaterialValidationLibrary::IsPackageInPaths(PackagePath, ExcludePaths))
 	{
-		if (PackagePath.StartsWith(Path))
-		{
-			return true;
-		}
+		return true;
 	}
 
-	for (FString const& Path : SearchPaths)
+	if (MaterialValidationLibrary::IsPackageInPaths(PackagePath, SearchPaths))
 	{
-		if (PackagePath.StartsWith(Path))
-		{
-			return false;
-		}
+		return false;
 	}
 
 	return true;
