Change Material Validation Groups to use FString paths instead of SoftObjectPtr
This avoids asset referencing which is causing production issues when deleting UMaterial assets.
#rnx

diff --git a/Engine/Plugins/Editor/MaterialValidation/Source/Private/MaterialValidationGroup.h b/Engine/Plugins/Editor/MaterialValidation/Source/Private/MaterialValidationGroup.h
--- a/Engine/Plugins/Editor/MaterialValidation/Source/Private/MaterialValidationGroup.h
+++ b/Engine/Plugins/Editor/MaterialValidation/Source/Private/MaterialValidationGroup.h
@@ -45,9 +45,9 @@ public:
 	UPROPERTY(EditAnywhere, Category = Material, meta = (DisplayName = "Include All Project Plugin Content"))
 	bool bAddProjectPluginPaths;
 
-	/** Map of UMaterial assets to validation descriptions. */
-	UPROPERTY(EditAnywhere, Category = Material, meta = (DisplayThumbnail = False, HidePinAssetPicker = True))
-	TMap<TSoftObjectPtr<UMaterial>, FMaterialValidationDesc> Materials;
+	/** Map of UMaterial asset paths to validation descriptions. Paths are FString to avoid references to the actual assets which would need managing on asset delete etc. */
+	UPROPERTY(VisibleAnywhere, Category = Material, meta = (DisplayThumbnail = False, HidePinAssetPicker = True))
+	TMap<FString, FMaterialValidationDesc> Materials;
 
 	/** Remove all Material assets that cannot be found on disk. */
 	UFUNCTION(CallInEditor, Category = Actions, meta = (DisplayName = "Remove Missing Materials"))
diff --git a/Engine/Plugins/Editor/MaterialValidation/Source/Private/MaterialValidationLibrary.cpp b/Engine/Plugins/Editor/MaterialValidation/Source/Private/MaterialValidationLibrary.cpp
--- a/Engine/Plugins/Editor/MaterialValidation/Source/Private/MaterialValidationLibrary.cpp
+++ b/Engine/Plugins/Editor/MaterialValidation/Source/Private/MaterialValidationLibrary.cpp
@@ -143,6 +143,40 @@ static void GetAssetsFromSearchPaths(UClass* InClass, TConstArrayView<FString> I
 	}
 }
 
+/** Convert a string path to a FSoftObjectPath taking into account any asset redirectors. */
+FSoftObjectPath ConvertStringPathToSoftObjectPath(FString const& InAssetPath)
+{
+	FSoftObjectPath CurrentPath(InAssetPath);
+
+	FAssetRegistryModule& AssetRegistryModule = FModuleManager::LoadModuleChecked<FAssetRegistryModule>("AssetRegistry");
+	IAssetRegistry& AssetRegistry = AssetRegistryModule.Get();
+
+	const int32 MaxRedirects = 10;
+	for (int32 RedirectDepth = 0; RedirectDepth < MaxRedirects; ++RedirectDepth)
+	{
+		FAssetData AssetData = AssetRegistry.GetAssetByObjectPath(CurrentPath);
+		if (!AssetData.IsValid())
+		{
+			break;
+		}
+
+		if (!AssetData.IsRedirector())
+		{
+			break;
+		}
+
+		FString RedirectedPath;
+		if (!AssetData.GetTagValue("DestinationObject", RedirectedPath))
+		{
+			break;
+		}
+
+		CurrentPath = FSoftObjectPath(RedirectedPath);
+	}
+
+	return CurrentPath;
+}
+
 /** 
  * Load the instance hierarchy for a material and add all permutations to the FMaterialValidationDesc. 
  * This is potentially very slow for a material with many references that may also load/compile textures etc.
@@ -274,7 +308,10 @@ static void GetNonValidMaterialInstances(UMaterial const& InMaterial, FMaterialV
 static bool IsValidationSuccess(UMaterialValidationGroup const& InGroup, UMaterialInstanceConstant const& InMaterialInstance)
 {
 	UMaterial const* Material = InMaterialInstance.GetMaterial();
-	FMaterialValidationDesc const* Settings = InGroup.Materials.Find(Material);
+	const FSoftObjectPath SoftPackagePath(Material);
+	const FString PackagePath = SoftPackagePath.ToString();
+
+	FMaterialValidationDesc const* Settings = InGroup.Materials.Find(PackagePath);
 	if (Settings == nullptr)
 	{
 		return true;
@@ -353,7 +390,10 @@ void UMaterialValidationLibrary::AddMaterialsToGroup(UMaterialValidationGroup* I
 
 		if (Material != nullptr)
 		{
-			FMaterialValidationDesc& Desc = InGroup->Materials.FindOrAdd(Material);
+			const FSoftObjectPath SoftPackagePath(Material);
+			const FString PackagePath = SoftPackagePath.ToString();
+
+			FMaterialValidationDesc& Desc = InGroup->Materials.FindOrAdd(PackagePath);
 			if (bInUpdatePermutations)
 			{
 				MaterialValidationLibrary::UpdateValidMaterialInstances(*Material, Desc);
@@ -385,12 +425,12 @@ void UMaterialValidationLibrary::AddMaterialsFromGroupDirectoryPaths(UMaterialVa
 	int32 bNumAdded = 0;
 	for (FAssetData const& AssetData : AssetDatas)
 	{
-		TSoftObjectPtr<UMaterial> MaterialSoftPtr(AssetData.ToSoftObjectPath());
+		FString MaterialPath = AssetData.ToSoftObjectPath().ToString();
 
-		if (!InGroup->Materials.Contains(MaterialSoftPtr))
+		if (!InGroup->Materials.Contains(MaterialPath))
 		{
 			UE_LOG(LogMaterialValidation, Log, TEXT("Adding material '%s' to '%s'."), *AssetData.AssetName.ToString(), *InGroup->GetName());
-			InGroup->Materials.Add(MaterialSoftPtr);
+			InGroup->Materials.Add(MaterialPath);
 			bNumAdded++;
 		}
 	}
@@ -412,23 +452,25 @@ void UMaterialValidationLibrary::RemoveMissingMaterialsFromGroup(UMaterialValida
 	TArray<FString> ExcludePaths;
 	MaterialValidationLibrary::GetSearchPaths(*InGroup, SearchPaths, ExcludePaths);
 
-	TArray<TSoftObjectPtr<UMaterial>> ToRemove;
-	for (TPair<TSoftObjectPtr<UMaterial>, FMaterialValidationDesc>& It : InGroup->Materials)
+	TArray<FString> ToRemove;
+	for (TPair<FString, FMaterialValidationDesc>& It : InGroup->Materials)
 	{
-		FString PackageName = It.Key.ToSoftObjectPath().GetLongPackageName();
-		if (!FPackageName::DoesPackageExist(PackageName))
+		const FSoftObjectPath SoftPackagePath = MaterialValidationLibrary::ConvertStringPathToSoftObjectPath(It.Key);
+		const FString PackagePath = SoftPackagePath.ToString();
+
+		if (!FPackageName::DoesPackageExist(SoftPackagePath.GetLongPackageName()))
 		{
 			ToRemove.Add(It.Key);
 			continue;
 		}
 
-		if (MaterialValidationLibrary::IsPackageInPaths(PackageName, ExcludePaths))
+		if (MaterialValidationLibrary::IsPackageInPaths(PackagePath, ExcludePaths))
 		{
 			ToRemove.Add(It.Key);
 			continue;
 		}
 
-		if (!MaterialValidationLibrary::IsPackageInPaths(PackageName, SearchPaths))
+		if (!MaterialValidationLibrary::IsPackageInPaths(PackagePath, SearchPaths))
 		{
 			ToRemove.Add(It.Key);
 		}
@@ -436,9 +478,9 @@ void UMaterialValidationLibrary::RemoveMissingMaterialsFromGroup(UMaterialValida
 
 	UE_LOG(LogMaterialValidation, Display, TEXT("Removing %d materials."), ToRemove.Num());
 
-	for (TSoftObjectPtr<UMaterial> const& Material : ToRemove)
+	for (FString const& Material : ToRemove)
 	{
-		UE_LOG(LogMaterialValidation, Log, TEXT("Removing material '%s' from '%s'."), *Material.GetAssetName(), *InGroup->GetName());
+		UE_LOG(LogMaterialValidation, Log, TEXT("Removing material '%s' from '%s'."), *Material, *InGroup->GetName());
 		InGroup->Materials.Remove(Material);
 	}
 
@@ -459,11 +501,12 @@ void UMaterialValidationLibrary::UpdateAllMaterialsInGroup(UMaterialValidationGr
 
 	FScopedSlowTask SlowTask(static_cast<float>(InGroup->Materials.Num()), LOCTEXT("UpdatingMaterials", "Updating Materials..."));
 	
-	for (TPair<TSoftObjectPtr<UMaterial>, FMaterialValidationDesc>& It : InGroup->Materials)
+	for (TPair<FString, FMaterialValidationDesc>& It : InGroup->Materials)
 	{
 		SlowTask.EnterProgressFrame();
 
-		if (UMaterial const* Material = It.Key.LoadSynchronous())
+		TSoftObjectPtr<UMaterial> MaterialSoftPtr(MaterialValidationLibrary::ConvertStringPathToSoftObjectPath(It.Key));
+		if (UMaterial const* Material = MaterialSoftPtr.LoadSynchronous())
 		{
 			MaterialValidationLibrary::UpdateValidMaterialInstances(*Material, It.Value);
 		}
@@ -483,11 +526,12 @@ void UMaterialValidationLibrary::ValidateAllMaterialsInGroup(UMaterialValidation
 
 	FScopedSlowTask SlowTask(static_cast<float>(InGroup->Materials.Num()), LOCTEXT("ValidatingMaterials", "Validating Materials..."));
 	
-	for (TPair<TSoftObjectPtr<UMaterial>, FMaterialValidationDesc>& It : InGroup->Materials)
+	for (TPair<FString, FMaterialValidationDesc>& It : InGroup->Materials)
 	{
 		SlowTask.EnterProgressFrame();
 
-		if (UMaterial const* Material = It.Key.LoadSynchronous())
+		TSoftObjectPtr<UMaterial> MaterialSoftPtr(MaterialValidationLibrary::ConvertStringPathToSoftObjectPath(It.Key));
+		if (UMaterial const* Material = MaterialSoftPtr.LoadSynchronous())
 		{
 			MaterialValidationLibrary::GetNonValidMaterialInstances(*Material, It.Value, OutMaterialInstances);
 		}
@@ -500,9 +544,11 @@ void UMaterialValidationLibrary::AddMaterialInstanceToGroup(UMaterialValidationG
 	{
 		if (UMaterialInstanceConstant* MaterialInstance = Cast<UMaterialInstanceConstant>(InMaterialInstance))
 		{
-			UMaterial* BaseMaterial = MaterialInstance->GetMaterial();
+			UMaterial const* BaseMaterial = MaterialInstance->GetMaterial();
+			const FSoftObjectPath SoftPackagePath(BaseMaterial);
+			const FString PackagePath = SoftPackagePath.ToString();
 
-			if (FMaterialValidationDesc* Settings = InGroup->Materials.Find(BaseMaterial))
+			if (FMaterialValidationDesc* Settings = InGroup->Materials.Find(PackagePath))
 			{
 				const uint32 Hash = MaterialValidationLibrary::BuildHash(*MaterialInstance, *BaseMaterial);
 				if (Hash != 0 && !Settings->ApprovedPermutationHashes.Contains(Hash))
@@ -537,13 +583,14 @@ bool UMaterialValidationLibrary::ValidateMaterialInGroup(UMaterialValidationGrou
 		return true;
 	}
 
-	if (InGroup->Materials.Contains(InMaterial))
+	const FSoftObjectPath SoftPackagePath(InMaterial);
+	const FString PackagePath = SoftPackagePath.ToString();
+
+	if (InGroup->Materials.Contains(PackagePath))
 	{
 		return true;
 	}
 
-	FString PackagePath = InMaterial->GetPathName();
-
 	TArray<FString> SearchPaths;
 	TArray<FString> ExcludePaths;
 	MaterialValidationLibrary::GetSearchPaths(*InGroup, SearchPaths, ExcludePaths);
@@ -568,7 +615,7 @@ int32 UMaterialValidationLibrary::GetApprovedPermutationCount(TSoftObjectPtr<UMa
 
 	for (UMaterialValidationGroup* Group : Groups)
 	{
-		if (FMaterialValidationDesc* Settings = Group->Materials.Find(InMaterial))
+		if (FMaterialValidationDesc* Settings = Group->Materials.Find(InMaterial.ToSoftObjectPath().ToString()))
 		{
 			return Settings->ApprovedPermutationHashes.Num();
 		}
