Extend material validation instance hash to take into account base material overrides, static component masks and material layers.
Add versioning to the material validation group and use it to invalidate data using previous hash version.
Existing assets will still need an update to actually use the new hash.
#rnx

diff --git a/Engine/Plugins/Editor/MaterialValidation/Source/Private/MaterialValidationGroup.cpp b/Engine/Plugins/Editor/MaterialValidation/Source/Private/MaterialValidationGroup.cpp
--- a/Engine/Plugins/Editor/MaterialValidation/Source/Private/MaterialValidationGroup.cpp
+++ b/Engine/Plugins/Editor/MaterialValidation/Source/Private/MaterialValidationGroup.cpp
@@ -4,9 +4,46 @@
 
 #include "MaterialValidationLibrary.h"
 #include "Misc/PackageName.h"
+#include "UObject/DevObjectVersion.h"
 
 #include UE_INLINE_GENERATED_CPP_BY_NAME(MaterialValidationGroup)
 
+/** Custom serialization version for MaterialValidationGroup. */
+struct FMaterialValidationGroupObjectVersion
+{
+	enum Type
+	{
+		// Before any version changes were made.
+		BeforeCustomVersionWasAdded = 0,
+		// Added base material overrides into permutation hash calculation.
+		UpdatedPermutationHash,
+
+		// New versions can be added above this line.
+		VersionPlusOne,
+		LatestVersion = VersionPlusOne - 1
+	};
+
+	// The GUID for this custom version number
+	const static FGuid GUID;
+
+	FMaterialValidationGroupObjectVersion() = delete;
+};
+
+const FGuid FMaterialValidationGroupObjectVersion::GUID(0x0D03FD52, 0xBF174EAE, 0x89C8DE2A, 0x98BF53FE);
+static FDevVersionRegistration GRegisterAtomFrameworkObjectVersion(FMaterialValidationGroupObjectVersion::GUID, FMaterialValidationGroupObjectVersion::LatestVersion, TEXT("MaterialValidationGroup"));
+
+void UMaterialValidationGroup::Serialize(FArchive& Ar)
+{
+	Super::Serialize(Ar);
+	Ar.UsingCustomVersion(FMaterialValidationGroupObjectVersion::GUID);
+
+	if (Ar.IsLoading() && Ar.CustomVer(FMaterialValidationGroupObjectVersion::GUID) < FMaterialValidationGroupObjectVersion::UpdatedPermutationHash)
+	{
+		// Clear invalidated data.
+		Materials.Reset();
+	}
+}
+
 // Path picker returns absolute paths. Convert to package paths such as /Game/...
 static void ConvertPaths(TArrayView<FDirectoryPath> Paths)
 {
diff --git a/Engine/Plugins/Editor/MaterialValidation/Source/Private/MaterialValidationGroup.h b/Engine/Plugins/Editor/MaterialValidation/Source/Private/MaterialValidationGroup.h
--- a/Engine/Plugins/Editor/MaterialValidation/Source/Private/MaterialValidationGroup.h
+++ b/Engine/Plugins/Editor/MaterialValidation/Source/Private/MaterialValidationGroup.h
@@ -68,5 +68,6 @@ public:
 	void Update();
 
 protected:
+	virtual void Serialize(FArchive& Ar) override;
 	virtual void PostEditChangeProperty(FPropertyChangedEvent& PropertyChangedEvent) override;
 };
diff --git a/Engine/Plugins/Editor/MaterialValidation/Source/Private/MaterialValidationLibrary.cpp b/Engine/Plugins/Editor/MaterialValidation/Source/Private/MaterialValidationLibrary.cpp
--- a/Engine/Plugins/Editor/MaterialValidation/Source/Private/MaterialValidationLibrary.cpp
+++ b/Engine/Plugins/Editor/MaterialValidation/Source/Private/MaterialValidationLibrary.cpp
@@ -8,6 +8,7 @@
 #include "Interfaces/IPluginManager.h"
 #include "Materials/Material.h"
 #include "Materials/MaterialInstanceConstant.h"
+#include "Materials/MaterialLayersFunctions.h"
 #include "MaterialValidationConfig.h"
 #include "MaterialValidationGroup.h"
 #include "MaterialValidationModule.h"
@@ -19,38 +20,89 @@
 
 namespace MaterialValidationLibrary {
 
-/** 
- * Build hash for material instance that should be shared by instances that generate the same shader permutations. 
- * todo: Add non-static switch contributors.
- */
-static uint32 BuildHash(UMaterialInstanceConstant const& InMaterialInstance, UMaterial const& InBaseMaterial, TConstArrayView<FMaterialParameterInfo> InStaticSwitchParameterInfo)
+/** Base material data necessary to calculate permutation id. */
+struct FMaterialPermutationData
+{
+	TArray<FMaterialParameterInfo> StaticSwitchParameterInfo;
+	TArray<FMaterialParameterInfo> ComponentMaskParameterInfo;
+};
+
+/** Build the FMaterialPermutationData for a UMaterial. */
+void BuildPermutationData(UMaterial const& InBaseMaterial, FMaterialPermutationData& OutPermutationData)
+{
+	TArray<FGuid> StaticSwitchGuids;
+	InBaseMaterial.GetAllStaticSwitchParameterInfo(OutPermutationData.StaticSwitchParameterInfo, StaticSwitchGuids);
+	TArray<FGuid> ComponentMaskParameterIds;
+	InBaseMaterial.GetAllStaticComponentMaskParameterInfo(OutPermutationData.ComponentMaskParameterInfo, ComponentMaskParameterIds);
+}
+
+/** Build hash for a material instance which matches for all instances that generate the same shader permutations. */
+uint32 BuildPermutationHash(UMaterialInstanceConstant const& InMaterialInstance, UMaterial const& InBaseMaterial, FMaterialPermutationData const& InPermutationData)
 {
 	FString KeyString;
-	for (int ParameterIndex = 0; ParameterIndex < InStaticSwitchParameterInfo.Num(); ++ParameterIndex)
+
+	// Append overridden material properties.
+	if (InMaterialInstance.HasOverridenBaseProperties())
+	{
+		FSHAHash Hash;
+		InMaterialInstance.GetBasePropertyOverridesHash(Hash);
+		KeyString.Append(Hash.ToString());
+	}
+
+	// Append overridden static switches.
+	for (int ParameterIndex = 0; ParameterIndex < InPermutationData.StaticSwitchParameterInfo.Num(); ++ParameterIndex)
 	{
 		FMaterialParameterMetadata InstanceMeta;
-		const bool bInstanceValid = InMaterialInstance.GetParameterValue(EMaterialParameterType::StaticSwitch, InStaticSwitchParameterInfo[ParameterIndex], InstanceMeta);
+		const bool bInstanceValid = InMaterialInstance.GetParameterValue(EMaterialParameterType::StaticSwitch, InPermutationData.StaticSwitchParameterInfo[ParameterIndex], InstanceMeta);
 		const bool Value = bInstanceValid && InstanceMeta.Value.AsStaticSwitch();
 		
 		FMaterialParameterMetadata BaseMeta;
-		const bool bBaseValid = InBaseMaterial.GetParameterValue(EMaterialParameterType::StaticSwitch, InStaticSwitchParameterInfo[ParameterIndex], BaseMeta);
+		const bool bBaseValid = InBaseMaterial.GetParameterValue(EMaterialParameterType::StaticSwitch, InPermutationData.StaticSwitchParameterInfo[ParameterIndex], BaseMeta);
 		const bool BaseValue = bBaseValid && BaseMeta.Value.AsStaticSwitch();
 
 		if (Value != BaseValue)
 		{
-			InStaticSwitchParameterInfo[ParameterIndex].AppendString(KeyString);
+			InPermutationData.StaticSwitchParameterInfo[ParameterIndex].AppendString(KeyString);
+		}
+	}
+
+	// Append overridden component masks.
+	for (int ParameterIndex = 0; ParameterIndex < InPermutationData.ComponentMaskParameterInfo.Num(); ++ParameterIndex)
+	{
+		FMaterialParameterMetadata InstanceMeta;
+		const bool bInstanceValid = InMaterialInstance.GetParameterValue(EMaterialParameterType::StaticComponentMask, InPermutationData.ComponentMaskParameterInfo[ParameterIndex], InstanceMeta);
+		const FStaticComponentMaskValue Value = bInstanceValid ? InstanceMeta.Value.AsStaticComponentMask() : FStaticComponentMaskValue();
+
+		FMaterialParameterMetadata BaseMeta;
+		const bool bBaseValid = InBaseMaterial.GetParameterValue(EMaterialParameterType::StaticComponentMask, InPermutationData.ComponentMaskParameterInfo[ParameterIndex], BaseMeta);
+		const FStaticComponentMaskValue BaseValue = bBaseValid ? BaseMeta.Value.AsStaticComponentMask() : FStaticComponentMaskValue();
+
+		if (Value != BaseValue)
+		{
+			InPermutationData.ComponentMaskParameterInfo[ParameterIndex].AppendString(KeyString);
+
+			const uint32 ValueMask = (Value.R ? 1 : 0) | (Value.G ? 2 : 0) | (Value.B ? 4 : 0) | (Value.A ? 8 : 0);
+			KeyString.Append(FString::Printf(TEXT("%d"), ValueMask));
 		}
 	}
+
+	// Append material layers.
+ 	FMaterialLayersFunctions MaterialLayers;
+ 	if (InMaterialInstance.GetMaterialLayers(MaterialLayers))
+ 	{
+		KeyString.Append(MaterialLayers.GetStaticPermutationString());
+ 	}
+
+	// Should return 0 if this isn't a static permutation (should be an empty KeyString).
 	return GetTypeHash(KeyString);
 }
 
-/** Build hash for material instance that should be shared by instances that generate the same shader permutations. */
-static uint32 BuildHash(UMaterialInstanceConstant const& InMaterialInstance, UMaterial const& InBaseMaterial)
+/** Build hash for a material instance which matches for all instances that generate the same shader permutations. */
+static uint32 BuildPermutationHash(UMaterialInstanceConstant const& InMaterialInstance, UMaterial const& InBaseMaterial)
 {
-	TArray<FMaterialParameterInfo> StaticSwitchParameterInfo;
-	TArray<FGuid> StaticSwitchGuids;
-	InBaseMaterial.GetAllStaticSwitchParameterInfo(StaticSwitchParameterInfo, StaticSwitchGuids);
-	return BuildHash(InMaterialInstance, InBaseMaterial, StaticSwitchParameterInfo);
+	FMaterialPermutationData PermutationData;
+	BuildPermutationData(InBaseMaterial, PermutationData);
+	return BuildPermutationHash(InMaterialInstance, InBaseMaterial, PermutationData);
 }
 
 /** Gather all search and exclude directory paths for a group. */
@@ -268,14 +320,13 @@ static void UpdateValidMaterialInstances(UMaterial const& InMaterial, FMaterialV
 	TArray<UMaterialInstanceConstant*> MaterialInstances;
 	LoadMaterialInstanceConstants(InMaterial, MaterialInstances);
 
-	TArray<FMaterialParameterInfo> StaticSwitchParameterInfo;
-	TArray<FGuid> StaticSwitchGuids;
-	InMaterial.GetAllStaticSwitchParameterInfo(StaticSwitchParameterInfo, StaticSwitchGuids);
+	FMaterialPermutationData PermutationData;
+	BuildPermutationData(InMaterial, PermutationData);
 
 	OutDesc.ApprovedPermutationHashes.Reset();
 	for (UMaterialInstanceConstant* MaterialInstance : MaterialInstances)
  	{
- 		const uint32 Hash = BuildHash(*MaterialInstance, InMaterial, StaticSwitchParameterInfo);
+ 		const uint32 Hash = BuildPermutationHash(*MaterialInstance, InMaterial, PermutationData);
 		if (Hash != 0 && !OutDesc.ApprovedPermutationHashes.Contains(Hash))
 		{
 			UE_LOG(LogMaterialValidation, Log, TEXT("Add material instance '%s' to approved list."), *MaterialInstance->GetName());
@@ -290,13 +341,12 @@ static void GetNonValidMaterialInstances(UMaterial const& InMaterial, FMaterialV
 	TArray<UMaterialInstanceConstant*> MaterialInstances;
 	LoadMaterialInstanceConstants(InMaterial, MaterialInstances);
 
-	TArray<FMaterialParameterInfo> StaticSwitchParameterInfo;
-	TArray<FGuid> StaticSwitchGuids;
-	InMaterial.GetAllStaticSwitchParameterInfo(StaticSwitchParameterInfo, StaticSwitchGuids);
+	FMaterialPermutationData PermutationData;
+	BuildPermutationData(InMaterial, PermutationData);
 
 	for (UMaterialInstanceConstant* MaterialInstance : MaterialInstances)
 	{
-		const uint32 Hash = BuildHash(*MaterialInstance, InMaterial, StaticSwitchParameterInfo);
+		const uint32 Hash = BuildPermutationHash(*MaterialInstance, InMaterial, PermutationData);
 		if (Hash != 0 && !InDesc.ApprovedPermutationHashes.Contains(Hash))
 		{
 			OutMaterialInstances.Add(MaterialInstance);
@@ -317,7 +367,7 @@ static bool IsValidationSuccess(UMaterialValidationGroup const& InGroup, UMateri
 		return true;
 	}
 
-	uint32 Hash = BuildHash(InMaterialInstance, *Material);
+	uint32 Hash = BuildPermutationHash(InMaterialInstance, *Material);
 	return Hash == 0 || Settings->ApprovedPermutationHashes.Contains(Hash);
 }
 
@@ -550,7 +600,7 @@ void UMaterialValidationLibrary::AddMaterialInstanceToGroup(UMaterialValidationG
 
 			if (FMaterialValidationDesc* Settings = InGroup->Materials.Find(PackagePath))
 			{
-				const uint32 Hash = MaterialValidationLibrary::BuildHash(*MaterialInstance, *BaseMaterial);
+				const uint32 Hash = MaterialValidationLibrary::BuildPermutationHash(*MaterialInstance, *BaseMaterial);
 				if (Hash != 0 && !Settings->ApprovedPermutationHashes.Contains(Hash))
 				{
 					UE_LOG(LogMaterialValidation, Log, TEXT("Add material instance '%s' to approved list."), *MaterialInstance->GetName());
diff --git a/Engine/Source/Runtime/Engine/Public/Materials/MaterialInstance.h b/Engine/Source/Runtime/Engine/Public/Materials/MaterialInstance.h
--- a/Engine/Source/Runtime/Engine/Public/Materials/MaterialInstance.h
+++ b/Engine/Source/Runtime/Engine/Public/Materials/MaterialInstance.h
@@ -1106,7 +1106,7 @@ public:
 	ENGINE_API virtual void UpdateCachedData();
 #endif
 
-	void GetBasePropertyOverridesHash(FSHAHash& OutHash)const;
+	ENGINE_API void GetBasePropertyOverridesHash(FSHAHash& OutHash)const;
 	ENGINE_API virtual bool HasOverridenBaseProperties()const;
 
 #if WITH_EDITOR
