Fix warnings when removing entries from removed plugins.
#rnx

diff --git a/Engine/Plugins/Editor/MaterialValidation/Source/Private/MaterialValidationLibrary.cpp b/Engine/Plugins/Editor/MaterialValidation/Source/Private/MaterialValidationLibrary.cpp
--- a/Engine/Plugins/Editor/MaterialValidation/Source/Private/MaterialValidationLibrary.cpp
+++ b/Engine/Plugins/Editor/MaterialValidation/Source/Private/MaterialValidationLibrary.cpp
@@ -506,21 +506,28 @@ void UMaterialValidationLibrary::RemoveMissingMaterialsFromGroup(UMaterialValida
 	for (TPair<FString, FMaterialValidationDesc>& It : InGroup->Materials)
 	{
 		const FSoftObjectPath SoftPackagePath = MaterialValidationLibrary::ConvertStringPathToSoftObjectPath(It.Key);
-		const FString PackagePath = SoftPackagePath.ToString();
 
-		if (!FPackageName::DoesPackageExist(SoftPackagePath.GetLongPackageName()))
+		FPackagePath PackagePath;
+		if (!FPackagePath::TryFromPackageName(SoftPackagePath.GetLongPackageName(), PackagePath))
 		{
 			ToRemove.Add(It.Key);
 			continue;
 		}
 
-		if (MaterialValidationLibrary::IsPackageInPaths(PackagePath, ExcludePaths))
+		if (!FPackageName::DoesPackageExist(PackagePath))
 		{
 			ToRemove.Add(It.Key);
 			continue;
 		}
 
-		if (!MaterialValidationLibrary::IsPackageInPaths(PackagePath, SearchPaths))
+		const FString PackageName = PackagePath.GetPackageName();
+		if (MaterialValidationLibrary::IsPackageInPaths(PackageName, ExcludePaths))
+		{
+			ToRemove.Add(It.Key);
+			continue;
+		}
+
+		if (!MaterialValidationLibrary::IsPackageInPaths(PackageName, SearchPaths))
 		{
 			ToRemove.Add(It.Key);
 		}
