Material permutation validators combined into a single validator that accumulates a shader count.
Adds a budget threshold in the config.
Logs specific info text for scraping by submit tools.
#rnx

diff --git a/Engine/Plugins/Editor/MaterialValidation/Source/Private/MaterialValidationConfig.h b/Engine/Plugins/Editor/MaterialValidation/Source/Private/MaterialValidationConfig.h
--- a/Engine/Plugins/Editor/MaterialValidation/Source/Private/MaterialValidationConfig.h
+++ b/Engine/Plugins/Editor/MaterialValidation/Source/Private/MaterialValidationConfig.h
@@ -9,15 +9,6 @@
 
 class UMaterialValidationGroup;
 
-/** Enumeration defining at what level validation is run. */
-UENUM()
-enum class EMaterialValidationLevel
-{
-	DisplayOnly UMETA(ToolTip = "Only show validation result in Material Editor UI."),
-	OnSave UMETA(ToolTip = "Show in Material Editor UI and validate on save."),
-	OnSubmit UMETA(ToolTip = "Show in Material Editor UI and validate on both save and submit."),
-};
-
 /** Configuration for the material validation system. */
 UCLASS(config = Editor, defaultconfig, meta = (DisplayName = "Data Validation Settings"), MinimalAPI)
 class UMaterialValidationConfig : public UDeveloperSettings
@@ -31,9 +22,11 @@ public:
 	/** Set true to enable validation. */
 	UPROPERTY(config, EditAnywhere, Category = MaterialPermutationBudget, meta = (DisplayName = "Enable Validation"))
 	bool bEnable = false;
-	/** Set levels where validation occurs. */
+
+	/** Threshold for shader count increase before triggering further submit approval. */
 	UPROPERTY(config, EditAnywhere, Category = MaterialPermutationBudget, meta = (EditCondition = "bEnable"))
-	EMaterialValidationLevel ValidationLevel = EMaterialValidationLevel::DisplayOnly;
+	int32 ShaderBudgetThreshold = 0;
+
 	/** An array of Material Validation Groups to use during validation. */
 	UPROPERTY(config, EditAnywhere, Category = MaterialPermutationBudget, meta = (EditCondition = "bEnable", DisplayName = "Material Validation Groups"))
 	TArray<TSoftObjectPtr<UMaterialValidationGroup>> Groups;
diff --git a/Engine/Plugins/Editor/MaterialValidation/Source/Private/MaterialValidationModule.cpp b/Engine/Plugins/Editor/MaterialValidation/Source/Private/MaterialValidationModule.cpp
--- a/Engine/Plugins/Editor/MaterialValidation/Source/Private/MaterialValidationModule.cpp
+++ b/Engine/Plugins/Editor/MaterialValidation/Source/Private/MaterialValidationModule.cpp
@@ -41,13 +41,9 @@ void FMaterialValidationModule::OnEditorInitialized(double Duration)
 
 		if (UEditorValidatorSubsystem* Subsystem = GEditor->GetEditorSubsystem<UEditorValidatorSubsystem>())
 		{
-			UEditorValidatorBase* ValidatorMaterialInstance = NewObject<UEditorValidator_MaterialInstancePermutation>();
-			Subsystem->AddValidator(ValidatorMaterialInstance);
-			EditorValidators.Add(ValidatorMaterialInstance);
-
-			UEditorValidatorBase* ValidatorMaterial = NewObject<UEditorValidator_MaterialPermutation>();
-			Subsystem->AddValidator(ValidatorMaterial);
-			EditorValidators.Add(ValidatorMaterial);
+			UEditorValidatorBase* Validator = NewObject<UEditorValidator_MaterialPermutation>();
+			Subsystem->AddValidator(Validator);
+			EditorValidators.Add(Validator);
 		}
 	}
 }
diff --git a/Engine/Plugins/Editor/MaterialValidation/Source/Private/MaterialValidators.cpp b/Engine/Plugins/Editor/MaterialValidation/Source/Private/MaterialValidators.cpp
--- a/Engine/Plugins/Editor/MaterialValidation/Source/Private/MaterialValidators.cpp
+++ b/Engine/Plugins/Editor/MaterialValidation/Source/Private/MaterialValidators.cpp
@@ -3,6 +3,7 @@
 #include "MaterialValidators.h"
 
 #include "AssetRegistry/AssetDataToken.h"
+#include "Editor.h"
 #include "Logging/TokenizedMessage.h"
 #include "MaterialShared.h"
 #include "Materials/Material.h"
@@ -10,73 +11,95 @@
 #include "MaterialValidationConfig.h"
 #include "MaterialValidationGroup.h"
 #include "MaterialValidationLibrary.h"
+#include "Misc/DataValidation.h"
 
 #include UE_INLINE_GENERATED_CPP_BY_NAME(MaterialValidators)
 
 #define LOCTEXT_NAMESPACE "MaterialValidation"
 
-bool UEditorValidator_MaterialPermutation::CanValidateAsset_Implementation(const FAssetData& InAssetData, UObject* InAsset, FDataValidationContext& InContext) const
+/**
+ * Get the number of shaders associated with a material instance.
+ * These _may_ be deduplicated with other material instances, so this is an upper
+ * limit estimate of the cost.
+ * todo: Calculate for other platforms, or at least the current preview platform.
+ * todo: Can we get an estimate of memory cost too?
+ */
+static int32 GetShaderCount(UMaterialInterface* InMaterial)
 {
-	return Cast<UMaterial>(InAsset) != nullptr;
-}
+	TArray<FDebugShaderTypeInfo> ShaderInfos;
+	InMaterial->GetShaderTypes(GMaxRHIShaderPlatform, nullptr, ShaderInfos);
 
-EDataValidationResult UEditorValidator_MaterialPermutation::ValidateLoadedAsset_Implementation(const FAssetData& InAssetData, UObject* InAsset, FDataValidationContext& InContext)
-{
-	if (UMaterial* Material = Cast<UMaterial>(InAsset))
+	int NumShaders = 0;
+	for (FDebugShaderTypeInfo const& ShaderInfo : ShaderInfos)
 	{
-		TArray<UMaterialValidationGroup*> Groups;
-		UMaterialValidationLibrary::GetAllGroups(Groups, /*bInSyncLoad*/true);
-
-		for (UMaterialValidationGroup* Group : Groups)
+		NumShaders += ShaderInfo.ShaderTypes.Num();
+		for (FDebugShaderPipelineInfo const& PipelineInfo : ShaderInfo.Pipelines)
 		{
-			if (!UMaterialValidationLibrary::ValidateMaterialInGroup(Group, Material))
-			{
-				UMaterialValidationConfig const* Config = GetDefault<UMaterialValidationConfig>();
-				const EMessageSeverity::Type MessageSeverity = Config->ValidationLevel == EMaterialValidationLevel::DisplayOnly ? EMessageSeverity::Info : EMessageSeverity::Error;
-
-				TWeakObjectPtr<UMaterial> WeakMaterialPtr = MakeWeakObjectPtr(Material);
-				TWeakObjectPtr<UMaterialValidationGroup> WeakGroupPtr = MakeWeakObjectPtr(Group);
-
-				TSharedRef<FTokenizedMessage> Message = AssetMessage(MessageSeverity);
-				Message->AddText(LOCTEXT("Fail_NewMaterial", "This is an unapproved new material in:"));
-				Message->AddToken(FAssetDataToken::Create(Group));
-			}
+			NumShaders += PipelineInfo.ShaderTypes.Num();
 		}
 	}
+	return NumShaders;
+}
 
-	if (GetValidationResult() != EDataValidationResult::Invalid)
+bool UEditorValidator_MaterialPermutation::CanValidateAsset_Implementation(const FAssetData& InAssetData, UObject* InAsset, FDataValidationContext& InContext) const
+{
+	const EDataValidationUsecase Usecase = InContext.GetValidationUsecase();
+	if (Usecase != EDataValidationUsecase::Commandlet && Usecase != EDataValidationUsecase::PreSubmit)
 	{
-		AssetPasses(InAsset);
+		return false;
 	}
-
-	return GetValidationResult();
+	if (Cast<UMaterialInstanceConstant>(InAsset) == nullptr && Cast<UMaterial>(InAsset) == nullptr)
+	{
+		return false;
+	}
+	return true;
 }
 
-bool UEditorValidator_MaterialInstancePermutation::CanValidateAsset_Implementation(const FAssetData& InAssetData, UObject* InAsset, FDataValidationContext& InContext) const
+EDataValidationResult UEditorValidator_MaterialPermutation::ValidateLoadedAsset_Implementation(const FAssetData& InAssetData, UObject* InAsset, FDataValidationContext& InContext)
 {
-	return Cast<UMaterialInstanceConstant>(InAsset) != nullptr;
-}
+	UMaterial* Material = Cast<UMaterial>(InAsset);
+	UMaterialInstance* MaterialInstance = Cast<UMaterialInstance>(InAsset);
 
-EDataValidationResult UEditorValidator_MaterialInstancePermutation::ValidateLoadedAsset_Implementation(const FAssetData& InAssetData, UObject* InAsset, FDataValidationContext& InContext)
-{
-	if (UMaterialInstance* MaterialInstance = Cast<UMaterialInstance>(InAsset))
+	if (Material || MaterialInstance)
 	{
 		TArray<UMaterialValidationGroup*> Groups;
 		UMaterialValidationLibrary::GetAllGroups(Groups, /*bInSyncLoad*/true);
 
+		int32 NumShaders = 0;
+
 		for (UMaterialValidationGroup* Group : Groups)
 		{
-			if (!UMaterialValidationLibrary::ValidateMaterialInstanceInGroup(Group, MaterialInstance))
+			if (Material && !UMaterialValidationLibrary::ValidateMaterialInGroup(Group, Material))
 			{
-				UMaterialValidationConfig const* Config = GetDefault<UMaterialValidationConfig>();
-				const EMessageSeverity::Type MessageSeverity = Config->ValidationLevel == EMaterialValidationLevel::DisplayOnly ? EMessageSeverity::Info : EMessageSeverity::Error;
+				NumShaders = GetShaderCount(Material);
+				break;
+			}
+			
+			if (MaterialInstance && !UMaterialValidationLibrary::ValidateMaterialInstanceInGroup(Group, MaterialInstance))
+			{
+				NumShaders = GetShaderCount(MaterialInstance);
+				break;
+			}
+		}
 
-				TWeakObjectPtr<UMaterialInstance> WeakMaterialInstancePtr = MakeWeakObjectPtr(MaterialInstance);
-				TWeakObjectPtr<UMaterialValidationGroup> WeakGroupPtr = MakeWeakObjectPtr(Group);
+		if (NumShaders > 0)
+		{
+			AccumulatedShaderCount += NumShaders;
+
+			// Need to register callback to reset AccumulatedShaderCount after validation is complete.
+			if (!OnPostValidationHandle.IsValid())
+			{
+				OnPostValidationHandle = FEditorDelegates::OnPostAssetValidation.AddUObject(this, &UEditorValidator_MaterialPermutation::OnPostValidation);
+			}
 
-				TSharedRef<FTokenizedMessage> Message = AssetMessage(MessageSeverity);
-				Message->AddText(LOCTEXT("Fail_NewPermutation", "The material instance creates a new material permutation in:"));
-				Message->AddToken(FAssetDataToken::Create(Group));
+			AssetMessage(EMessageSeverity::Info)->AddText(LOCTEXT("ShaderBudget_Info", "Adds new shaders: {0}"), FText::AsNumber(NumShaders));
+			AssetMessage(EMessageSeverity::Info)->AddText(LOCTEXT("ShaderBudget_Delta", "ShaderBudgetDelta: {0}"), FText::AsNumber(AccumulatedShaderCount));
+
+			UMaterialValidationConfig const* Config = GetDefault<UMaterialValidationConfig>();
+			const int32 Threshold = Config->ShaderBudgetThreshold;
+			if (Threshold > 0 && AccumulatedShaderCount >= Threshold)
+			{
+				AssetMessage(EMessageSeverity::Info)->AddText(LOCTEXT("ShaderBudget_IntegrationRequest", "ShaderBudgetIntegrationRequest"));
 			}
 		}
 	}
@@ -86,31 +109,16 @@ EDataValidationResult UEditorValidator_MaterialInstancePermutation::ValidateLoad
 		AssetPasses(InAsset);
 	}
 
-	return GetValidationResult(); 
+	return GetValidationResult();
 }
 
-/** 
- * Get the number of shaders associated with a material instance.
- * These _may_ be deduplicated with other material instances, so this is an upper
- * limit estimate of the cost.
- * todo: Calculate for other platforms, or at least the current preview platform.
- * todo: Can we get an estimate of memory cost too?
- */
-static int32 GetShaderCount(UMaterialInstance* InMaterialInstance)
+void UEditorValidator_MaterialPermutation::OnPostValidation()
 {
-	TArray<FDebugShaderTypeInfo> ShaderInfos;
-	InMaterialInstance->GetShaderTypes(GMaxRHIShaderPlatform, nullptr, ShaderInfos);
-	
-	int NumShaders = 0;
-	for (FDebugShaderTypeInfo const& ShaderInfo : ShaderInfos)
-	{
-		NumShaders += ShaderInfo.ShaderTypes.Num();
-		for (FDebugShaderPipelineInfo const& PipelineInfo : ShaderInfo.Pipelines)
-		{
-			NumShaders += PipelineInfo.ShaderTypes.Num();
-		}
-	}
-	return NumShaders;
+	AccumulatedShaderCount = 0;
+
+	// Remove delegate until next validation.
+	FEditorDelegates::OnPostAssetValidation.Remove(OnPostValidationHandle);
+	OnPostValidationHandle.Reset();
 }
 
 bool UMaterialEditorValidator_Permutation::Validate(UMaterialInterface* InMaterial, FMaterialEditorValidatorContext& InContext) const
diff --git a/Engine/Plugins/Editor/MaterialValidation/Source/Private/MaterialValidators.h b/Engine/Plugins/Editor/MaterialValidation/Source/Private/MaterialValidators.h
--- a/Engine/Plugins/Editor/MaterialValidation/Source/Private/MaterialValidators.h
+++ b/Engine/Plugins/Editor/MaterialValidation/Source/Private/MaterialValidators.h
@@ -15,16 +15,12 @@ class UEditorValidator_MaterialPermutation : public UEditorValidatorBase
 
 	virtual bool CanValidateAsset_Implementation(const FAssetData& InAssetData, UObject* InAsset, FDataValidationContext& InContext) const override;
 	virtual EDataValidationResult ValidateLoadedAsset_Implementation(const FAssetData& InAssetData, UObject* InAsset, FDataValidationContext& InContext) override;
-};
 
-/** Implementation of UEditorValidatorBase that applies validation for limiting the addition of material permutations from UMaterialInstances. */
-UCLASS()
-class UEditorValidator_MaterialInstancePermutation : public UEditorValidatorBase
-{
-	GENERATED_BODY()
+private:
+	int32 AccumulatedShaderCount = 0;
 
-	virtual bool CanValidateAsset_Implementation(const FAssetData& InAssetData, UObject* InAsset, FDataValidationContext& InContext) const override;
-	virtual EDataValidationResult ValidateLoadedAsset_Implementation(const FAssetData& InAssetData, UObject* InAsset, FDataValidationContext& InContext) override;
+	FDelegateHandle OnPostValidationHandle;
+	void OnPostValidation();
 };
 
 /** Implementation of UMaterialEditorValidatorBase that applies validation for limiting the addition of material permutations. */
