Use culture invariant FText for some material validation tags.
Tags that are for scraping by SubmitTool need to be invariant.
#rnx

diff --git a/Engine/Plugins/Editor/MaterialValidation/Source/Private/MaterialValidators.cpp b/Engine/Plugins/Editor/MaterialValidation/Source/Private/MaterialValidators.cpp
--- a/Engine/Plugins/Editor/MaterialValidation/Source/Private/MaterialValidators.cpp
+++ b/Engine/Plugins/Editor/MaterialValidation/Source/Private/MaterialValidators.cpp
@@ -92,14 +92,14 @@ EDataValidationResult UEditorValidator_MaterialPermutation::ValidateLoadedAsset_
 				OnPostValidationHandle = FEditorDelegates::OnPostAssetValidation.AddUObject(this, &UEditorValidator_MaterialPermutation::OnPostValidation);
 			}
 
-			AssetMessage(EMessageSeverity::Info)->AddText(LOCTEXT("ShaderBudget_Info", "Adds new shaders: {0}"), FText::AsNumber(NumShaders));
-			AssetMessage(EMessageSeverity::Info)->AddText(LOCTEXT("ShaderBudget_Delta", "ShaderBudgetDelta: {0}"), FText::AsNumber(AccumulatedShaderCount));
+			AssetMessage(EMessageSeverity::Info)->AddText(LOCTEXT("ShaderBudget_Info", "Adds {0} new shaders"), FText::AsNumber(NumShaders));
+			AssetMessage(EMessageSeverity::Info)->AddText(FText::Format(INVTEXT("ShaderBudgetDelta: {0}"), AccumulatedShaderCount));
 
 			UMaterialValidationConfig const* Config = GetDefault<UMaterialValidationConfig>();
 			const int32 Threshold = Config->ShaderBudgetThreshold;
 			if (Threshold > 0 && AccumulatedShaderCount >= Threshold)
 			{
-				AssetMessage(EMessageSeverity::Info)->AddText(LOCTEXT("ShaderBudget_IntegrationRequest", "ShaderBudgetIntegrationRequest"));
+				AssetMessage(EMessageSeverity::Info)->AddText(INVTEXT("ShaderBudgetIntegrationRequest"));
 			}
 		}
 	}
