Extend material validator to catch more shader permutation cases.
#rnx

diff --git a/Engine/Plugins/Editor/MaterialValidation/Source/Private/MaterialValidationAssetData.h b/Engine/Plugins/Editor/MaterialValidation/Source/Private/MaterialValidationAssetData.h
new file mode 100644
--- /dev/null
+++ b/Engine/Plugins/Editor/MaterialValidation/Source/Private/MaterialValidationAssetData.h
@@ -0,0 +1,60 @@
+// Copyright Epic Games, Inc. All Rights Reserved.
+
+#pragma once
+
+#include "Materials/MaterialInstanceBasePropertyOverrides.h"
+#include "MaterialValidationAssetData.generated.h"
+
+/** 
+ * Extracted data from UMaterial to be used for predicting the shader permutation ID of UMaterial and child UMaterialInstances. 
+ * This is intended for storage external to the UMaterial so that we don't need to load it.
+ * In future this may be better in the AssetData, but we will need to resave all assets before we can rely on that.
+ */
+USTRUCT(MinimalAPI)
+struct FMaterialValidationAssetData_Material
+{
+	GENERATED_BODY()
+
+	UPROPERTY()
+	FMaterialInstanceBasePropertyOverrides BaseProperties;
+	UPROPERTY()
+ 	TArray<uint32> StaticSwitchValues;
+ 	UPROPERTY()
+ 	TArray<uint32> ComponentMaskValues;
+	UPROPERTY()
+	uint32 StaticPropertyLayoutHash;
+	UPROPERTY()
+	uint32 MaterialLayerHash;
+	UPROPERTY()
+	uint32 PermutationHash;
+	UPROPERTY()
+	int32 ShaderCount;
+};
+
+/**
+ * Extracted data from UMaterialInstance to be used for predicting the shader permutation.
+ * This is intended for storage external to the UMaterialInstance so that we don't need to load it.
+ * In future this may be better in the AssetData, but we will need to resave all assets before we can rely on that.
+ */
+USTRUCT(MinimalAPI)
+struct FMaterialValidationAssetData_MaterialInstance
+{
+	GENERATED_BODY()
+
+	UPROPERTY()
+	FMaterialInstanceBasePropertyOverrides BasePropertyOverrides;
+ 	UPROPERTY()
+ 	TArray<uint32> StaticSwitchOverrideMask;
+ 	UPROPERTY()
+ 	TArray<uint32> StaticSwitchOverrideValues;
+ 	UPROPERTY()
+ 	TArray<uint32> ComponentMaskOverrideMask;
+ 	UPROPERTY()
+ 	TArray<uint32> ComponentMaskOverrideValues;
+	UPROPERTY()
+	uint32 MaterialLayerHash;
+	UPROPERTY()
+	uint32 PermutationHash;
+	UPROPERTY()
+	int32 ShaderCount;
+};
diff --git a/Engine/Plugins/Editor/MaterialValidation/Source/Private/MaterialValidationGroup.cpp b/Engine/Plugins/Editor/MaterialValidation/Source/Private/MaterialValidationGroup.cpp
--- a/Engine/Plugins/Editor/MaterialValidation/Source/Private/MaterialValidationGroup.cpp
+++ b/Engine/Plugins/Editor/MaterialValidation/Source/Private/MaterialValidationGroup.cpp
@@ -17,6 +17,8 @@ struct FMaterialValidationGroupObjectVersion
 		BeforeCustomVersionWasAdded = 0,
 		// Added base material overrides into permutation hash calculation.
 		UpdatedPermutationHash,
+		// Refined base material overrides to those that are relevant for shader permutations.
+		UpdatedPermutationHash2,
 
 		// New versions can be added above this line.
 		VersionPlusOne,
@@ -37,7 +39,7 @@ void UMaterialValidationGroup::Serialize(FArchive& Ar)
 	Super::Serialize(Ar);
 	Ar.UsingCustomVersion(FMaterialValidationGroupObjectVersion::GUID);
 
-	if (Ar.IsLoading() && Ar.CustomVer(FMaterialValidationGroupObjectVersion::GUID) < FMaterialValidationGroupObjectVersion::UpdatedPermutationHash)
+	if (Ar.IsLoading() && Ar.CustomVer(FMaterialValidationGroupObjectVersion::GUID) < FMaterialValidationGroupObjectVersion::UpdatedPermutationHash2)
 	{
 		// Clear invalidated data.
 		Materials.Reset();
diff --git a/Engine/Plugins/Editor/MaterialValidation/Source/Private/MaterialValidationGroup.h b/Engine/Plugins/Editor/MaterialValidation/Source/Private/MaterialValidationGroup.h
--- a/Engine/Plugins/Editor/MaterialValidation/Source/Private/MaterialValidationGroup.h
+++ b/Engine/Plugins/Editor/MaterialValidation/Source/Private/MaterialValidationGroup.h
@@ -3,6 +3,7 @@
 #pragma once
 
 #include "Engine/DataAsset.h"
+#include "MaterialValidationAssetData.h"
 #include "MaterialValidationGroup.generated.h"
 
 class UMaterial;
@@ -12,9 +13,13 @@ USTRUCT(MinimalAPI)
 struct FMaterialValidationDesc
 {
 	GENERATED_BODY()
+	
+	/** The approved base material permutation. */
+	UPROPERTY()
+	uint32 BaseMaterialHash;
 
-	/** Array of currently approved material permutations. */
-	UPROPERTY(VisibleAnywhere, Category = Material)
+	/** Array of currently approved material permutations not including the BaseMaterialHash. */
+	UPROPERTY()
 	TArray<uint32> ApprovedPermutationHashes;
 };
 
@@ -25,10 +30,6 @@ class UMaterialValidationGroup : public UDataAsset
 	GENERATED_BODY()
 
 public:
-	/** Array of users that are approved to modify the Material assets in Materials and under the MaterialPaths. */
-	UPROPERTY(EditAnywhere, Category = Material)
-	TArray<FString> Approvers;
-
 	/** Array of explicit directory paths that we use when searching for UMaterial packages that need validating. */
 	UPROPERTY(EditAnywhere, Category = Material)
 	TArray<FDirectoryPath> MaterialPaths;
@@ -46,7 +47,7 @@ public:
 	bool bAddProjectPluginPaths;
 
 	/** Map of UMaterial asset paths to validation descriptions. Paths are FString to avoid references to the actual assets which would need managing on asset delete etc. */
-	UPROPERTY(VisibleAnywhere, Category = Material, meta = (DisplayThumbnail = False, HidePinAssetPicker = True))
+	UPROPERTY(VisibleAnywhere, Category = Material)
 	TMap<FString, FMaterialValidationDesc> Materials;
 
 	/** Remove all Material assets that cannot be found on disk. */
diff --git a/Engine/Plugins/Editor/MaterialValidation/Source/Private/MaterialValidationLibrary.cpp b/Engine/Plugins/Editor/MaterialValidation/Source/Private/MaterialValidationLibrary.cpp
--- a/Engine/Plugins/Editor/MaterialValidation/Source/Private/MaterialValidationLibrary.cpp
+++ b/Engine/Plugins/Editor/MaterialValidation/Source/Private/MaterialValidationLibrary.cpp
@@ -9,6 +9,7 @@
 #include "Materials/Material.h"
 #include "Materials/MaterialInstanceConstant.h"
 #include "Materials/MaterialLayersFunctions.h"
+#include "MaterialShared.h"
 #include "MaterialValidationConfig.h"
 #include "MaterialValidationGroup.h"
 #include "MaterialValidationModule.h"
@@ -23,86 +24,329 @@ namespace MaterialValidationLibrary {
 /** Base material data necessary to calculate permutation id. */
 struct FMaterialPermutationData
 {
+	/** Infos required to request per material instance static switch values. */
 	TArray<FMaterialParameterInfo> StaticSwitchParameterInfo;
+	/** Infos required to request per material instance component mask values. */
 	TArray<FMaterialParameterInfo> ComponentMaskParameterInfo;
+	/** Hash of all static switch and component mask GUIDs. */
+	uint32 StaticPropertyLayoutHash;
 };
 
 /** Build the FMaterialPermutationData for a UMaterial. */
-void BuildPermutationData(UMaterial const& InBaseMaterial, FMaterialPermutationData& OutPermutationData)
+static void BuildPermutationData(UMaterial const& InBaseMaterial, FMaterialPermutationData& OutPermutationData)
 {
 	TArray<FGuid> StaticSwitchGuids;
 	InBaseMaterial.GetAllStaticSwitchParameterInfo(OutPermutationData.StaticSwitchParameterInfo, StaticSwitchGuids);
-	TArray<FGuid> ComponentMaskParameterIds;
-	InBaseMaterial.GetAllStaticComponentMaskParameterInfo(OutPermutationData.ComponentMaskParameterInfo, ComponentMaskParameterIds);
+	
+	TArray<FGuid> ComponentMaskGuids;
+	InBaseMaterial.GetAllStaticComponentMaskParameterInfo(OutPermutationData.ComponentMaskParameterInfo, ComponentMaskGuids);
+
+	OutPermutationData.StaticPropertyLayoutHash = HashCombine(GetTypeHash(StaticSwitchGuids), GetTypeHash(ComponentMaskGuids));
 }
 
-/** Build hash for a material instance which matches for all instances that generate the same shader permutations. */
-uint32 BuildPermutationHash(UMaterialInstanceConstant const& InMaterialInstance, UMaterial const& InBaseMaterial, FMaterialPermutationData const& InPermutationData)
+/** 
+ * Initialize the FMaterialInstanceBasePropertyOverrides values (but not override flags) from a material.
+ * We only consider the parts of FMaterialInstanceBasePropertyOverrides that trigger shader permutations.
+ * We could use a more compact struct with only these parts, but we just reuse a sparse part of this existing struct.
+ */
+static FMaterialInstanceBasePropertyOverrides InitRelevantBaseProperties(UMaterialInterface const& InMaterial)
 {
-	FString KeyString;
+	FMaterialInstanceBasePropertyOverrides BaseProperties;
+	BaseProperties.OpacityMaskClipValue = InMaterial.GetOpacityMaskClipValue();
+	BaseProperties.BlendMode = InMaterial.GetBlendMode();
+	BaseProperties.ShadingModel = InMaterial.GetShadingModels().GetFirstShadingModel();
+	BaseProperties.DitheredLODTransition = InMaterial.IsDitheredLODTransition();
+	BaseProperties.bCastDynamicShadowAsMasked = InMaterial.GetCastDynamicShadowAsMasked();
+	BaseProperties.TwoSided = InMaterial.IsTwoSided();
+	BaseProperties.bIsThinSurface = InMaterial.IsThinSurface();
+	BaseProperties.bOutputTranslucentVelocity = InMaterial.IsTranslucencyWritingVelocity();
+	BaseProperties.bHasPixelAnimation = InMaterial.HasPixelAnimation();
+	BaseProperties.bEnableTessellation = InMaterial.IsTessellationEnabled();
+	BaseProperties.UsageFlags = InMaterial.GetUsageFlags();
+	return BaseProperties;
+}
+
+/** Copy the relevant parts of the FMaterialInstanceBasePropertyOverrides struct from the value stored on the Material Instance. */
+static FMaterialInstanceBasePropertyOverrides InitRelevantBasePropertyOverrides(UMaterialInstance const& InMaterialInstance)
+{
+	FMaterialInstanceBasePropertyOverrides BasePropertyOverrides;
+	BasePropertyOverrides.bOverride_OpacityMaskClipValue = InMaterialInstance.BasePropertyOverrides.bOverride_OpacityMaskClipValue;
+	BasePropertyOverrides.OpacityMaskClipValue = InMaterialInstance.BasePropertyOverrides.OpacityMaskClipValue;
+	BasePropertyOverrides.bOverride_BlendMode = InMaterialInstance.BasePropertyOverrides.bOverride_BlendMode;
+	BasePropertyOverrides.BlendMode = InMaterialInstance.BasePropertyOverrides.BlendMode;
+	BasePropertyOverrides.bOverride_ShadingModel = InMaterialInstance.BasePropertyOverrides.bOverride_ShadingModel;
+	BasePropertyOverrides.ShadingModel = InMaterialInstance.BasePropertyOverrides.ShadingModel;
+	BasePropertyOverrides.bOverride_DitheredLODTransition = InMaterialInstance.BasePropertyOverrides.bOverride_DitheredLODTransition;
+	BasePropertyOverrides.DitheredLODTransition = InMaterialInstance.BasePropertyOverrides.DitheredLODTransition;
+	BasePropertyOverrides.bCastDynamicShadowAsMasked = InMaterialInstance.BasePropertyOverrides.bCastDynamicShadowAsMasked;
+	BasePropertyOverrides.bOverride_CastDynamicShadowAsMasked = InMaterialInstance.BasePropertyOverrides.bOverride_CastDynamicShadowAsMasked;
+	BasePropertyOverrides.bOverride_TwoSided = InMaterialInstance.BasePropertyOverrides.bOverride_TwoSided;
+	BasePropertyOverrides.TwoSided = InMaterialInstance.BasePropertyOverrides.TwoSided;
+	BasePropertyOverrides.bOverride_bIsThinSurface = InMaterialInstance.BasePropertyOverrides.bOverride_bIsThinSurface;
+	BasePropertyOverrides.bIsThinSurface = InMaterialInstance.BasePropertyOverrides.bIsThinSurface;
+	BasePropertyOverrides.bOverride_OutputTranslucentVelocity = InMaterialInstance.BasePropertyOverrides.bOverride_OutputTranslucentVelocity;
+	BasePropertyOverrides.bOutputTranslucentVelocity = InMaterialInstance.BasePropertyOverrides.bOutputTranslucentVelocity;
+	BasePropertyOverrides.bOverride_bHasPixelAnimation = InMaterialInstance.BasePropertyOverrides.bOverride_bHasPixelAnimation;
+	BasePropertyOverrides.bHasPixelAnimation = InMaterialInstance.BasePropertyOverrides.bHasPixelAnimation;
+	BasePropertyOverrides.bOverride_bEnableTessellation = InMaterialInstance.BasePropertyOverrides.bOverride_bEnableTessellation;
+	BasePropertyOverrides.bEnableTessellation = InMaterialInstance.BasePropertyOverrides.bEnableTessellation;
+	BasePropertyOverrides.bOverride_UsageFlags = InMaterialInstance.BasePropertyOverrides.bOverride_UsageFlags;
+	BasePropertyOverrides.UsageFlags = InMaterialInstance.BasePropertyOverrides.UsageFlags;
+	return BasePropertyOverrides;
+}
+
+/** Return a hash built from the relevant FMaterialInstanceBasePropertyOverrides values (but not override flags). */
+static uint32 BuildBasePropertiesHash(FMaterialInstanceBasePropertyOverrides const& InBaseProperties)
+{
+	const uint8 PackedBools = 
+		(InBaseProperties.TwoSided ? 1 << 0 : 0) |
+		(InBaseProperties.bIsThinSurface ? 1 << 1 : 0) |
+		(InBaseProperties.DitheredLODTransition ? 1 << 2 : 0) |
+		(InBaseProperties.bCastDynamicShadowAsMasked ? 1 << 3 : 0) |
+		(InBaseProperties.bOutputTranslucentVelocity ? 1 << 4 : 0) |
+		(InBaseProperties.bHasPixelAnimation ? 1 << 5 : 0) |
+		(InBaseProperties.bEnableTessellation ? 1 << 6 : 0);
+
+	uint32 Hash = GetTypeHash(PackedBools);
+	Hash = HashCombineFast(Hash, GetTypeHash(InBaseProperties.OpacityMaskClipValue));
+	Hash = HashCombineFast(Hash, GetTypeHash(InBaseProperties.BlendMode));
+	Hash = HashCombineFast(Hash, GetTypeHash(InBaseProperties.ShadingModel));
+	Hash = HashCombineFast(Hash, GetTypeHash(InBaseProperties.UsageFlags));
+	return Hash;
+}
 
-	// Append overridden material properties.
-	if (InMaterialInstance.HasOverridenBaseProperties())
+/** Return a hash built from any material layer settings on a material. */
+static uint32 BuildLayerHash(UMaterialInterface const& InMaterial)
+{
+	FMaterialLayersFunctions MaterialLayers;
+	if (InMaterial.GetMaterialLayers(MaterialLayers))
 	{
-		FSHAHash Hash;
-		InMaterialInstance.GetBasePropertyOverridesHash(Hash);
-		KeyString.Append(Hash.ToString());
+		return GetTypeHash(MaterialLayers.GetStaticPermutationString());
 	}
+	return 0;
+}
 
-	// Append overridden static switches.
-	for (int ParameterIndex = 0; ParameterIndex < InPermutationData.StaticSwitchParameterInfo.Num(); ++ParameterIndex)
+/** Return a hash built for a material or material instance which will match for all instances that generate the same shader permutations. */
+static uint32 BuildPermutationHash(FMaterialInstanceBasePropertyOverrides const& InBaseProperties, TConstArrayView<uint32> InStaticSwitchValues, TConstArrayView<uint32> InComponentMaskValues, uint32 InMaterialLayerHash)
+{
+	uint32 Hash = BuildBasePropertiesHash(InBaseProperties);
+	for (uint32 Value : InStaticSwitchValues)
 	{
-		FMaterialParameterMetadata InstanceMeta;
-		const bool bInstanceValid = InMaterialInstance.GetParameterValue(EMaterialParameterType::StaticSwitch, InPermutationData.StaticSwitchParameterInfo[ParameterIndex], InstanceMeta);
-		const bool Value = bInstanceValid && InstanceMeta.Value.AsStaticSwitch();
-		
-		FMaterialParameterMetadata BaseMeta;
-		const bool bBaseValid = InBaseMaterial.GetParameterValue(EMaterialParameterType::StaticSwitch, InPermutationData.StaticSwitchParameterInfo[ParameterIndex], BaseMeta);
-		const bool BaseValue = bBaseValid && BaseMeta.Value.AsStaticSwitch();
+		Hash = HashCombineFast(Hash, GetTypeHash(Value));
+	} 
+	for (uint32 Value : InComponentMaskValues)
+	{
+		Hash = HashCombineFast(Hash, GetTypeHash(Value));
+	}
+	Hash = HashCombineFast(Hash, InMaterialLayerHash);
+	return Hash;
+}
 
-		if (Value != BaseValue)
+/** 
+ * Get the (non-deduplicated) number of shaders that a single material will generat. 
+ * Note: This is only calculated for the current (editor) platform and probably contains debug shaders that aren't cooked.
+ */
+static int32 GetShaderCount(UMaterialInterface const& InMaterial)
+{
+	// todo: GetShaderTypes() isn't const, but maybe could be?
+	TArray<FDebugShaderTypeInfo> ShaderInfos;
+	(const_cast<UMaterialInterface&>(InMaterial)).GetShaderTypes(GMaxRHIShaderPlatform, nullptr, ShaderInfos);
+
+	int NumShaders = 0;
+	for (FDebugShaderTypeInfo const& ShaderInfo : ShaderInfos)
+	{
+		NumShaders += ShaderInfo.ShaderTypes.Num();
+		for (FDebugShaderPipelineInfo const& PipelineInfo : ShaderInfo.Pipelines)
 		{
-			InPermutationData.StaticSwitchParameterInfo[ParameterIndex].AppendString(KeyString);
+			NumShaders += PipelineInfo.ShaderTypes.Num();
 		}
 	}
+	return NumShaders;
+}
 
-	// Append overridden component masks.
-	for (int ParameterIndex = 0; ParameterIndex < InPermutationData.ComponentMaskParameterInfo.Num(); ++ParameterIndex)
+/** Extract AssetData from a UMaterial object. */
+static void GetAssetData(
+	UMaterial const& InMaterial, 
+	FMaterialPermutationData const& InPermutationData, 
+	FMaterialValidationAssetData_Material& OutAssetData)
+{
+	// Gather static switches.
 	{
-		FMaterialParameterMetadata InstanceMeta;
-		const bool bInstanceValid = InMaterialInstance.GetParameterValue(EMaterialParameterType::StaticComponentMask, InPermutationData.ComponentMaskParameterInfo[ParameterIndex], InstanceMeta);
-		const FStaticComponentMaskValue Value = bInstanceValid ? InstanceMeta.Value.AsStaticComponentMask() : FStaticComponentMaskValue();
+		const uint32 StaticSwitchNum = InPermutationData.StaticSwitchParameterInfo.Num();
+		const uint32 StaticSwitchBitArrayDwordNum = (StaticSwitchNum + 31) / 32;
+		OutAssetData.StaticSwitchValues.AddZeroed(StaticSwitchBitArrayDwordNum);
+
+		for (uint32 Index = 0; Index < StaticSwitchNum; ++Index)
+		{
+			FMaterialParameterMetadata Meta;
+			const bool bValid = InMaterial.GetParameterValue(EMaterialParameterType::StaticSwitch, InPermutationData.StaticSwitchParameterInfo[Index], Meta);
+			const bool bValue = bValid && Meta.Value.AsStaticSwitch();
+			if (bValue)
+			{
+				const uint32 BitArrayIndex = Index / 32;
+				const uint32 BitIndex = Index % 32;
+				OutAssetData.StaticSwitchValues[BitArrayIndex] |= 1 << BitIndex;
+			}
+		}
+	}
+
+	// Gather static component masks.
+	{
+		const uint32 ComponentMaskNum = InPermutationData.ComponentMaskParameterInfo.Num();
+		const uint32 ComponentMaskBitArrayDwordNum = ((ComponentMaskNum * 4) + 31) / 32;
+		OutAssetData.ComponentMaskValues.AddZeroed(ComponentMaskBitArrayDwordNum);
+
+		for (uint32 Index = 0; Index < ComponentMaskNum; ++Index)
+		{
+			FMaterialParameterMetadata Meta;
+			const bool bValid = InMaterial.GetParameterValue(EMaterialParameterType::StaticComponentMask, InPermutationData.ComponentMaskParameterInfo[Index], Meta);
+			const FStaticComponentMaskValue Value = bValid ? Meta.Value.AsStaticComponentMask() : FStaticComponentMaskValue();
+			const uint32 ValueMask = (Value.R ? 1 : 0) | (Value.G ? 2 : 0) | (Value.B ? 4 : 0) | (Value.A ? 8 : 0);
+
+			if (ValueMask)
+			{
+				const uint32 BitArrayIndex = Index / 8;
+				const uint32 BitIndex = (Index * 4) % 32;
+				OutAssetData.ComponentMaskValues[BitArrayIndex] |= ValueMask << BitIndex;
+			}
+		}
+	}
+
+	OutAssetData.BaseProperties = InitRelevantBaseProperties(InMaterial);
 
-		FMaterialParameterMetadata BaseMeta;
-		const bool bBaseValid = InBaseMaterial.GetParameterValue(EMaterialParameterType::StaticComponentMask, InPermutationData.ComponentMaskParameterInfo[ParameterIndex], BaseMeta);
-		const FStaticComponentMaskValue BaseValue = bBaseValid ? BaseMeta.Value.AsStaticComponentMask() : FStaticComponentMaskValue();
+	OutAssetData.MaterialLayerHash = BuildLayerHash(InMaterial);
+	
+	OutAssetData.StaticPropertyLayoutHash = InPermutationData.StaticPropertyLayoutHash;
+	OutAssetData.PermutationHash = BuildPermutationHash(OutAssetData.BaseProperties, OutAssetData.StaticSwitchValues, OutAssetData.ComponentMaskValues, OutAssetData.MaterialLayerHash);
+	OutAssetData.ShaderCount = GetShaderCount(InMaterial);
+}
+
+/** Extract AssetData from a UMaterialInstance object. */
+static void GetAssetData(
+	UMaterialInstance const& InMaterialInstance, 
+	UMaterialInterface const& InParentMaterial, 
+	FMaterialPermutationData const& InPermutationData,
+	FMaterialValidationAssetData_MaterialInstance& OutAssetData)
+{
+	// We collect the static switch and component mask values to create the
+	// permutation hash but only store the overridden values to the asset data.
+	TArray<uint32> StaticSwitchValues;
+	TArray<uint32> ComponentMaskValues;
+
+	// Gather static switches.
+	{
+		const uint32 StaticSwitchNum = InPermutationData.StaticSwitchParameterInfo.Num();
+		const uint32 StaticSwitchBitArrayDwordNum = (StaticSwitchNum + 31) / 32;
+		OutAssetData.StaticSwitchOverrideMask.AddZeroed(StaticSwitchBitArrayDwordNum);
+		OutAssetData.StaticSwitchOverrideValues.AddZeroed(StaticSwitchBitArrayDwordNum);
+		StaticSwitchValues.AddZeroed(StaticSwitchBitArrayDwordNum);
 
-		if (Value != BaseValue)
+		for (uint32 Index = 0; Index < StaticSwitchNum; ++Index)
 		{
-			InPermutationData.ComponentMaskParameterInfo[ParameterIndex].AppendString(KeyString);
+			FMaterialParameterMetadata InstanceMeta;
+			const bool bInstanceValid = InMaterialInstance.GetParameterValue(EMaterialParameterType::StaticSwitch, InPermutationData.StaticSwitchParameterInfo[Index], InstanceMeta);
+			const bool Value = bInstanceValid && InstanceMeta.Value.AsStaticSwitch();
+
+			FMaterialParameterMetadata ParentMeta;
+			const bool bBaseValid = InParentMaterial.GetParameterValue(EMaterialParameterType::StaticSwitch, InPermutationData.StaticSwitchParameterInfo[Index], ParentMeta);
+			const bool ParentValue = bBaseValid && ParentMeta.Value.AsStaticSwitch();
+
+			const uint32 BitArrayIndex = Index / 32;
+			const uint32 BitIndex = Index % 32;
+
+			if (Value)
+			{
+				StaticSwitchValues[BitArrayIndex] |= Value << BitIndex;
+			}
+			if (Value != ParentValue)
+			{
+				OutAssetData.StaticSwitchOverrideMask[BitArrayIndex] |= 1 << BitIndex;
+				OutAssetData.StaticSwitchOverrideValues[BitArrayIndex] |= Value << BitIndex;
+			}
+		}
+	}
+
+	// Gather static component masks.
+	{
+		const uint32 ComponentMaskNum = InPermutationData.ComponentMaskParameterInfo.Num();
+		const uint32 ComponentMaskBitValueArrayDwordNum = ((ComponentMaskNum * 4) + 31) / 32;
+		OutAssetData.ComponentMaskOverrideValues.AddZeroed(ComponentMaskBitValueArrayDwordNum);
+		OutAssetData.ComponentMaskOverrideMask.AddZeroed(ComponentMaskBitValueArrayDwordNum);
+		ComponentMaskValues.AddZeroed(ComponentMaskBitValueArrayDwordNum);
 
+		for (uint32 Index = 0; Index < ComponentMaskNum; ++Index)
+		{
+			FMaterialParameterMetadata InstanceMeta;
+			const bool bInstanceValid = InMaterialInstance.GetParameterValue(EMaterialParameterType::StaticComponentMask, InPermutationData.ComponentMaskParameterInfo[Index], InstanceMeta);
+			const FStaticComponentMaskValue Value = bInstanceValid ? InstanceMeta.Value.AsStaticComponentMask() : FStaticComponentMaskValue();
 			const uint32 ValueMask = (Value.R ? 1 : 0) | (Value.G ? 2 : 0) | (Value.B ? 4 : 0) | (Value.A ? 8 : 0);
-			KeyString.Append(FString::Printf(TEXT("%d"), ValueMask));
+
+			FMaterialParameterMetadata ParentMeta;
+			const bool bParentValid = InParentMaterial.GetParameterValue(EMaterialParameterType::StaticComponentMask, InPermutationData.ComponentMaskParameterInfo[Index], ParentMeta);
+			const FStaticComponentMaskValue ParentValue = bParentValid ? ParentMeta.Value.AsStaticComponentMask() : FStaticComponentMaskValue();
+
+			const uint32 BitArrayIndex = Index / 8;
+			const uint32 BitIndex = (Index * 4) % 32;
+
+			if (ValueMask)
+			{
+				ComponentMaskValues[BitArrayIndex] |= ValueMask << BitIndex;
+			}
+			if (Value != ParentValue)
+			{
+				OutAssetData.ComponentMaskOverrideMask[BitArrayIndex] |= 0xf << BitIndex;
+				OutAssetData.ComponentMaskOverrideValues[BitArrayIndex] |= ValueMask << BitIndex;
+			}
 		}
 	}
 
-	// Append material layers.
- 	FMaterialLayersFunctions MaterialLayers;
- 	if (InMaterialInstance.GetMaterialLayers(MaterialLayers))
- 	{
-		KeyString.Append(MaterialLayers.GetStaticPermutationString());
- 	}
+	FMaterialInstanceBasePropertyOverrides BaseProperties = InitRelevantBaseProperties(InMaterialInstance);
+	OutAssetData.BasePropertyOverrides = InitRelevantBasePropertyOverrides(InMaterialInstance);
 
-	// Should return 0 if this isn't a static permutation (should be an empty KeyString).
-	return GetTypeHash(KeyString);
+	OutAssetData.MaterialLayerHash = BuildLayerHash(InMaterialInstance);
+	
+	OutAssetData.PermutationHash = BuildPermutationHash(BaseProperties, StaticSwitchValues, ComponentMaskValues, OutAssetData.MaterialLayerHash);
+	OutAssetData.ShaderCount = GetShaderCount(InMaterialInstance);
 }
 
-/** Build hash for a material instance which matches for all instances that generate the same shader permutations. */
+/** Build hash for a material instance which will match for all instances that generate the same shader permutations. */
 static uint32 BuildPermutationHash(UMaterialInstanceConstant const& InMaterialInstance, UMaterial const& InBaseMaterial)
 {
 	FMaterialPermutationData PermutationData;
 	BuildPermutationData(InBaseMaterial, PermutationData);
-	return BuildPermutationHash(InMaterialInstance, InBaseMaterial, PermutationData);
+
+	FMaterialValidationAssetData_MaterialInstance InstanceAssetData;
+	GetAssetData(InMaterialInstance, *InMaterialInstance.Parent, PermutationData, InstanceAssetData);
+
+	return InstanceAssetData.PermutationHash;
+}
+
+/** Apply overrides for bit mask arrays such as those used by static switch and component masks. */
+static void ApplyBitMaskOverrides(TConstArrayView<uint32> InParentValues, TConstArrayView<uint32> InChildOverridesMasks, TConstArrayView<uint32> InChildOverrideValues, TArrayView<uint32> OutValues)
+{
+	const int32 NumDwords = InParentValues.Num();
+	check(NumDwords == InChildOverridesMasks.Num() && NumDwords == InChildOverrideValues.Num() && NumDwords == OutValues.Num());
+
+	for (int32 Index = 0; Index < NumDwords; ++Index)
+	{
+		OutValues[Index] = (InParentValues[Index] & ~InChildOverridesMasks[Index]) | (InChildOverridesMasks[Index] & InChildOverrideValues[Index]);
+	}
+}
+
+/** Apply material base property overrides to generate the child base property values from the parent values and the child overrides. */
+static void ApplyBasePropertyOverrides(FMaterialInstanceBasePropertyOverrides const& InParentBaseProperties, FMaterialInstanceBasePropertyOverrides const& InChildBasePropertyOverrides, FMaterialInstanceBasePropertyOverrides& OutBaseProperties)
+{
+	OutBaseProperties.TwoSided = InChildBasePropertyOverrides.bOverride_TwoSided ? InChildBasePropertyOverrides.TwoSided : InParentBaseProperties.TwoSided;
+	OutBaseProperties.bIsThinSurface = InChildBasePropertyOverrides.bOverride_bIsThinSurface ? InChildBasePropertyOverrides.bIsThinSurface : InParentBaseProperties.bIsThinSurface;
+	OutBaseProperties.DitheredLODTransition = InChildBasePropertyOverrides.bOverride_DitheredLODTransition ? InChildBasePropertyOverrides.DitheredLODTransition : InParentBaseProperties.DitheredLODTransition;
+	OutBaseProperties.bCastDynamicShadowAsMasked = InChildBasePropertyOverrides.bOverride_CastDynamicShadowAsMasked ? InChildBasePropertyOverrides.bCastDynamicShadowAsMasked : InParentBaseProperties.bCastDynamicShadowAsMasked;
+	OutBaseProperties.bOutputTranslucentVelocity = InChildBasePropertyOverrides.bOverride_OutputTranslucentVelocity ? InChildBasePropertyOverrides.bOutputTranslucentVelocity : InParentBaseProperties.bOutputTranslucentVelocity;
+	OutBaseProperties.bHasPixelAnimation = InChildBasePropertyOverrides.bOverride_bHasPixelAnimation ? InChildBasePropertyOverrides.bHasPixelAnimation : InParentBaseProperties.bHasPixelAnimation;
+	OutBaseProperties.bEnableTessellation = InChildBasePropertyOverrides.bOverride_bEnableTessellation ? InChildBasePropertyOverrides.bEnableTessellation : InParentBaseProperties.bEnableTessellation;
+	OutBaseProperties.BlendMode = InChildBasePropertyOverrides.bOverride_BlendMode ? InChildBasePropertyOverrides.BlendMode : InParentBaseProperties.BlendMode;
+	OutBaseProperties.ShadingModel = InChildBasePropertyOverrides.bOverride_ShadingModel ? InChildBasePropertyOverrides.ShadingModel : InParentBaseProperties.ShadingModel;
+	OutBaseProperties.OpacityMaskClipValue = InChildBasePropertyOverrides.bOverride_OpacityMaskClipValue ? InChildBasePropertyOverrides.OpacityMaskClipValue : InParentBaseProperties.OpacityMaskClipValue;
+	OutBaseProperties.UsageFlags = (InParentBaseProperties.UsageFlags & ~InChildBasePropertyOverrides.bOverride_UsageFlags) | (InChildBasePropertyOverrides.bOverride_UsageFlags & InChildBasePropertyOverrides.UsageFlags);
 }
 
 /** Gather all search and exclude directory paths for a group. */
@@ -196,7 +440,7 @@ static void GetAssetsFromSearchPaths(UClass* InClass, TConstArrayView<FString> I
 }
 
 /** Convert a string path to a FSoftObjectPath taking into account any asset redirectors. */
-FSoftObjectPath ConvertStringPathToSoftObjectPath(FString const& InAssetPath)
+static FSoftObjectPath ConvertStringPathToSoftObjectPath(FString const& InAssetPath)
 {
 	FSoftObjectPath CurrentPath(InAssetPath);
 
@@ -323,11 +567,19 @@ static void UpdateValidMaterialInstances(UMaterial const& InMaterial, FMaterialV
 	FMaterialPermutationData PermutationData;
 	BuildPermutationData(InMaterial, PermutationData);
 
+	FMaterialValidationAssetData_Material AssetData;
+	GetAssetData(InMaterial, PermutationData, AssetData);
+	OutDesc.BaseMaterialHash = AssetData.PermutationHash;
+
 	OutDesc.ApprovedPermutationHashes.Reset();
+
 	for (UMaterialInstanceConstant* MaterialInstance : MaterialInstances)
- 	{
- 		const uint32 Hash = BuildPermutationHash(*MaterialInstance, InMaterial, PermutationData);
-		if (Hash != 0 && !OutDesc.ApprovedPermutationHashes.Contains(Hash))
+	{
+		FMaterialValidationAssetData_MaterialInstance InstanceAssetData;
+		GetAssetData(*MaterialInstance, *MaterialInstance->Parent, PermutationData, InstanceAssetData);
+
+		const uint32 Hash = InstanceAssetData.PermutationHash;
+		if (Hash != OutDesc.BaseMaterialHash && !OutDesc.ApprovedPermutationHashes.Contains(Hash))
 		{
 			UE_LOG(LogMaterialValidation, Log, TEXT("Add material instance '%s' to approved list."), *MaterialInstance->GetName());
 			OutDesc.ApprovedPermutationHashes.Add(Hash);
@@ -346,8 +598,11 @@ static void GetNonValidMaterialInstances(UMaterial const& InMaterial, FMaterialV
 
 	for (UMaterialInstanceConstant* MaterialInstance : MaterialInstances)
 	{
-		const uint32 Hash = BuildPermutationHash(*MaterialInstance, InMaterial, PermutationData);
-		if (Hash != 0 && !InDesc.ApprovedPermutationHashes.Contains(Hash))
+		FMaterialValidationAssetData_MaterialInstance InstanceAssetData;
+		GetAssetData(*MaterialInstance, *MaterialInstance->Parent, PermutationData, InstanceAssetData);
+
+		const uint32 Hash = InstanceAssetData.PermutationHash;
+		if (Hash != InDesc.BaseMaterialHash && !InDesc.ApprovedPermutationHashes.Contains(Hash))
 		{
 			OutMaterialInstances.Add(MaterialInstance);
 		}
@@ -367,8 +622,8 @@ static bool IsValidationSuccess(UMaterialValidationGroup const& InGroup, UMateri
 		return true;
 	}
 
-	uint32 Hash = BuildPermutationHash(InMaterialInstance, *Material);
-	return Hash == 0 || Settings->ApprovedPermutationHashes.Contains(Hash);
+	const uint32 Hash = BuildPermutationHash(InMaterialInstance, *Material);
+	return Hash == Settings->BaseMaterialHash || Settings->ApprovedPermutationHashes.Contains(Hash);
 }
 
 } //namespace MaterialValidationLibrary
@@ -419,7 +674,6 @@ void UMaterialValidationLibrary::ResetGroup(UMaterialValidationGroup* InGroup)
 	if (InGroup != nullptr)
 	{
 		InGroup->Materials.Reset();
-		InGroup->Approvers.Reset();
 	}
 }
 
@@ -608,7 +862,7 @@ void UMaterialValidationLibrary::AddMaterialInstanceToGroup(UMaterialValidationG
 			if (FMaterialValidationDesc* Settings = InGroup->Materials.Find(PackagePath))
 			{
 				const uint32 Hash = MaterialValidationLibrary::BuildPermutationHash(*MaterialInstance, *BaseMaterial);
-				if (Hash != 0 && !Settings->ApprovedPermutationHashes.Contains(Hash))
+				if (Hash != Settings->BaseMaterialHash && !Settings->ApprovedPermutationHashes.Contains(Hash))
 				{
 					UE_LOG(LogMaterialValidation, Log, TEXT("Add material instance '%s' to approved list."), *MaterialInstance->GetName());
 
