Material permutation validation improvements.
Database now stores asset data on all material and material instances.
With this it can walk the material hierarcy on validation to determine a better shader delta.
Shader delta now takes into account how modifying override values can cascade down the hierarchy.
#rnx
#Synced-CL 48135599

diff --git a/Engine/Plugins/Editor/MaterialValidation/Source/Private/MaterialValidationAssetData.h b/Engine/Plugins/Editor/MaterialValidation/Source/Private/MaterialValidationAssetData.h
--- a/Engine/Plugins/Editor/MaterialValidation/Source/Private/MaterialValidationAssetData.h
+++ b/Engine/Plugins/Editor/MaterialValidation/Source/Private/MaterialValidationAssetData.h
@@ -5,30 +5,98 @@
 #include "Materials/MaterialInstanceBasePropertyOverrides.h"
 #include "MaterialValidationAssetData.generated.h"
 
+/** 
+ * Material instance properties that trigger a static permutation. 
+ * This is a selected subset of FMaterialInstanceBasePropertyOverrides and may need updating if anything is added to that struct.
+ */
+USTRUCT()
+struct FStaticPermutationProperties
+{
+	GENERATED_BODY()
+
+	UPROPERTY()
+	uint8 TwoSided : 1;
+	UPROPERTY()
+	uint8 bIsThinSurface : 1;
+	UPROPERTY()
+	uint8 DitheredLODTransition : 1;
+	UPROPERTY()
+	uint8 bCastDynamicShadowAsMasked:1;
+	UPROPERTY()
+	uint8 bOutputTranslucentVelocity : 1;
+	UPROPERTY()
+	uint8 bHasPixelAnimation : 1;
+	UPROPERTY()
+	uint8 bEnableTessellation : 1;
+	UPROPERTY()
+	TEnumAsByte<EBlendMode> BlendMode;
+	UPROPERTY()
+	TEnumAsByte<EMaterialShadingModel> ShadingModel;
+	UPROPERTY()
+	float OpacityMaskClipValue;
+	UPROPERTY()
+	uint32 UsageFlags;
+
+	FStaticPermutationProperties();
+};
+
+/**
+ * Flags that enable override of the properties in FStaticPermutationProperties.
+ * This is a selected subset of FMaterialInstanceBasePropertyOverrides and may need updating if anything is added to that struct.
+ */
+USTRUCT()
+struct FStaticPermutationPropertyOverrideFlags
+{
+	GENERATED_BODY()
+
+	UPROPERTY()
+	uint8 bOverride_OpacityMaskClipValue : 1;
+	UPROPERTY()
+	uint8 bOverride_BlendMode : 1;
+	UPROPERTY()
+	uint8 bOverride_ShadingModel : 1;
+	UPROPERTY()
+	uint8 bOverride_DitheredLODTransition : 1;
+	UPROPERTY()
+	uint8 bOverride_CastDynamicShadowAsMasked : 1;
+	UPROPERTY()
+	uint8 bOverride_TwoSided : 1;
+	UPROPERTY()
+	uint8 bOverride_bIsThinSurface : 1;
+	UPROPERTY()
+	uint8 bOverride_OutputTranslucentVelocity : 1;
+	UPROPERTY()
+	uint8 bOverride_bHasPixelAnimation : 1;
+	UPROPERTY()
+	uint8 bOverride_bEnableTessellation : 1;
+	UPROPERTY()
+	uint32 Override_UsageFlags;
+
+	FStaticPermutationPropertyOverrideFlags();
+};
+
 /** 
  * Extracted data from UMaterial to be used for predicting the shader permutation ID of UMaterial and child UMaterialInstances. 
  * This is intended for storage external to the UMaterial so that we don't need to load it.
  * In future this may be better in the AssetData, but we will need to resave all assets before we can rely on that.
  */
-USTRUCT(MinimalAPI)
+USTRUCT()
 struct FMaterialValidationAssetData_Material
 {
 	GENERATED_BODY()
 
-	UPROPERTY()
-	FMaterialInstanceBasePropertyOverrides BaseProperties;
 	UPROPERTY()
  	TArray<uint32> StaticSwitchValues;
  	UPROPERTY()
  	TArray<uint32> ComponentMaskValues;
 	UPROPERTY()
-	uint32 StaticPropertyLayoutHash;
+	FStaticPermutationProperties StaticProperties;
 	UPROPERTY()
-	uint32 MaterialLayerHash;
+	uint32 StaticPropertyLayoutHash = 0;
 	UPROPERTY()
-	uint32 PermutationHash;
+	uint32 MaterialLayerHash = 0;
 	UPROPERTY()
-	int32 ShaderCount;
+	uint32 PermutationHash = 0;
 };
 
 /**
@@ -36,13 +104,11 @@ struct FMaterialValidationAssetData_Material
  * This is intended for storage external to the UMaterialInstance so that we don't need to load it.
  * In future this may be better in the AssetData, but we will need to resave all assets before we can rely on that.
  */
-USTRUCT(MinimalAPI)
+USTRUCT()
 struct FMaterialValidationAssetData_MaterialInstance
 {
 	GENERATED_BODY()
 
-	UPROPERTY()
-	FMaterialInstanceBasePropertyOverrides BasePropertyOverrides;
  	UPROPERTY()
  	TArray<uint32> StaticSwitchOverrideMask;
  	UPROPERTY()
@@ -52,9 +118,11 @@ struct FMaterialValidationAssetData_MaterialInstance
  	UPROPERTY()
  	TArray<uint32> ComponentMaskOverrideValues;
 	UPROPERTY()
-	uint32 MaterialLayerHash;
+	FStaticPermutationProperties StaticProperties;
+	UPROPERTY()
+	FStaticPermutationPropertyOverrideFlags StaticPropertyOverrideFlags;
 	UPROPERTY()
-	uint32 PermutationHash;
+	uint32 MaterialLayerHash = 0;
 	UPROPERTY()
-	int32 ShaderCount;
+	uint32 PermutationHash = 0;
 };
diff --git a/Engine/Plugins/Editor/MaterialValidation/Source/Private/MaterialValidationCommandlet.cpp b/Engine/Plugins/Editor/MaterialValidation/Source/Private/MaterialValidationCommandlet.cpp
--- a/Engine/Plugins/Editor/MaterialValidation/Source/Private/MaterialValidationCommandlet.cpp
+++ b/Engine/Plugins/Editor/MaterialValidation/Source/Private/MaterialValidationCommandlet.cpp
@@ -23,25 +23,15 @@ void UMaterialValidationCommandlet::InitializePackageNames(const TArray<FString>
 	for (FString const& Switch : Switches)
 	{
 		int32 Value;
-		if (FParse::Value(*Switch, TEXT("RemoveMissing="), Value))
+		if (FParse::Value(*Switch, TEXT("UpdateMaterials="), Value))
 		{
-			bRemoveMissingMaterials = Value != 0;
-			UE_LOG(LogMaterialValidation, Log, TEXT("RemoveMissing: %d."), Value);
+			bUpdateMaterials = Value != 0;
+			UE_LOG(LogMaterialValidation, Log, TEXT("UpdateMaterials: %d."), Value);
 		}
-		if (FParse::Value(*Switch, TEXT("AddFromDirectories="), Value))
+		else if (FParse::Value(*Switch, TEXT("UpdatePermutations="), Value))
 		{
-			bAddMaterialsFromDirectories = Value != 0;
-			UE_LOG(LogMaterialValidation, Log, TEXT("AddFromDirectories: %d."), Value);
-		}
-		else if (FParse::Value(*Switch, TEXT("DumpNonApproved="), Value))
-		{
-			bDumpNonApprovedMaterialInstances = Value != 0;
-			UE_LOG(LogMaterialValidation, Log, TEXT("DumpNonApproved: %d."), Value);
-		}
-		else if (FParse::Value(*Switch, TEXT("UpdateApproved="), Value))
-		{
-			bUpdateApprovedMaterialsInstances = Value != 0;
-			UE_LOG(LogMaterialValidation, Log, TEXT("UpdateApproved: %d."), Value);
+			bUpdateMaterialPermutations = Value != 0;
+			UE_LOG(LogMaterialValidation, Log, TEXT("UpdatePermutations: %d."), Value);
 		}
 	}
 
@@ -108,41 +98,22 @@ void UMaterialValidationCommandlet::PerformAdditionalOperations(UObject* Object,
 
 	UE_LOG(LogMaterialValidation, Display, TEXT("Commandlet Begin '%s'"), *Group->GetName());
 
-	if (bRemoveMissingMaterials)
-	{
-		UE_LOG(LogMaterialValidation, Display, TEXT("Removing Materials."));
-
-		UMaterialValidationLibrary::RemoveMissingMaterialsFromGroup(Group);
-	}
-
-	if (bAddMaterialsFromDirectories)
+	if (bUpdateMaterials)
 	{
-		UE_LOG(LogMaterialValidation, Display, TEXT("Adding Materials."));
+		UE_LOG(LogMaterialValidation, Display, TEXT("Removing Invalid Materials."));
+		UMaterialValidationLibrary::RemoveInvalidMaterialsFromGroup(Group);
 
-		UMaterialValidationLibrary::AddMaterialsFromGroupDirectoryPaths(Group);
+		UE_LOG(LogMaterialValidation, Display, TEXT("Adding Missing Materials."));
+		UMaterialValidationLibrary::AddMissingMaterialsToGroup(Group);
 	}
 
-	if (bDumpNonApprovedMaterialInstances)
+	if (bUpdateMaterialPermutations)
 	{
-		UE_LOG(LogMaterialValidation, Display, TEXT("Finding Non Approved Material Instances."));
-
-		TArray<UMaterialInstance*> MaterialInstances;
-		UMaterialValidationLibrary::ValidateAllMaterialsInGroup(Group, MaterialInstances);
-
-		for (UMaterialInstance* MaterialInstance : MaterialInstances)
-		{
-			UE_LOG(LogMaterialValidation, Display, TEXT("Found new material instance '%s' in group '%s'"), *MaterialInstance->GetName(), *Group->GetName());
-		}
-	}
-
-	if (bUpdateApprovedMaterialsInstances)
-	{
-		UE_LOG(LogMaterialValidation, Display, TEXT("Updating Approved Material Instances."));
-
-		UMaterialValidationLibrary::UpdateAllMaterialsInGroup(Group);
+		UE_LOG(LogMaterialValidation, Display, TEXT("Updating Material Permutations."));
+		UMaterialValidationLibrary::UpdateMaterialPermutationsInGroup(Group);
 	}
 
-	bSavePackage = bSavePackage || bRemoveMissingMaterials || bAddMaterialsFromDirectories || bUpdateApprovedMaterialsInstances;
+	bSavePackage = bSavePackage || bUpdateMaterials || bUpdateMaterialPermutations;
 
 	UE_LOG(LogMaterialValidation, Display, TEXT("Commandlet End '%s'"), *Group->GetName());
 }
diff --git a/Engine/Plugins/Editor/MaterialValidation/Source/Private/MaterialValidationCommandlet.h b/Engine/Plugins/Editor/MaterialValidation/Source/Private/MaterialValidationCommandlet.h
--- a/Engine/Plugins/Editor/MaterialValidation/Source/Private/MaterialValidationCommandlet.h
+++ b/Engine/Plugins/Editor/MaterialValidation/Source/Private/MaterialValidationCommandlet.h
@@ -18,12 +18,8 @@ class UMaterialValidationCommandlet : public UBaseIteratePackagesCommandlet
 	//~ End UBaseIteratePackagesCommandlet Interface
 
 private:
-	/** true to remove any materials that are not found on disk. Disabled with RemoveMissing=0 on commandline. */
-	bool bRemoveMissingMaterials = true;
-	/** true to add any new materials found in the material directories. Disabled with AddFromDirectories=0 on commandline. */
-	bool bAddMaterialsFromDirectories = true;
-	/** true to search for non-approved material instances and dump them to the output. Enabled with DumpNonApproved=1 on commandline. */
-	bool bDumpNonApprovedMaterialInstances = false;
-	/** true to search for all material instances and add them to the approved list. Disabled with UpdateApproved=0 on commandline. */
-	bool bUpdateApprovedMaterialsInstances = true;
+	/** true to update the materials bu removing those not found on disk and adding any materials found in the search paths. Disabled with UpdateMaterials=0 on commandline. */
+	bool bUpdateMaterials = true;
+	/** true to search for all material instances and add them to the approved list. Disabled with UpdatePermutations=0 on commandline. */
+	bool bUpdateMaterialPermutations = true;
 };
diff --git a/Engine/Plugins/Editor/MaterialValidation/Source/Private/MaterialValidationGroup.cpp b/Engine/Plugins/Editor/MaterialValidation/Source/Private/MaterialValidationGroup.cpp
--- a/Engine/Plugins/Editor/MaterialValidation/Source/Private/MaterialValidationGroup.cpp
+++ b/Engine/Plugins/Editor/MaterialValidation/Source/Private/MaterialValidationGroup.cpp
@@ -19,6 +19,8 @@ struct FMaterialValidationGroupObjectVersion
 		UpdatedPermutationHash,
 		// Refined base material overrides to those that are relevant for shader permutations.
 		UpdatedPermutationHash2,
+		// Added asset data to give more accurate shader count reporting.
+		AddedAssetData,
 
 		// New versions can be added above this line.
 		VersionPlusOne,
@@ -34,12 +36,42 @@ struct FMaterialValidationGroupObjectVersion
 const FGuid FMaterialValidationGroupObjectVersion::GUID(0x0D03FD52, 0xBF174EAE, 0x89C8DE2A, 0x98BF53FE);
 static FDevVersionRegistration GRegisterAtomFrameworkObjectVersion(FMaterialValidationGroupObjectVersion::GUID, FMaterialValidationGroupObjectVersion::LatestVersion, TEXT("MaterialValidationGroup"));
 
+FStaticPermutationProperties::FStaticPermutationProperties()
+	: TwoSided(0)
+	, bIsThinSurface(0)
+	, DitheredLODTransition(0)
+	, bCastDynamicShadowAsMasked(0)
+	, bOutputTranslucentVelocity(0)
+	, bHasPixelAnimation(0)
+	, bEnableTessellation(0)
+	, BlendMode(BLEND_Opaque)
+	, ShadingModel(MSM_DefaultLit)
+	, OpacityMaskClipValue(.333333f)
+	, UsageFlags(0)
+{
+}
+
+FStaticPermutationPropertyOverrideFlags::FStaticPermutationPropertyOverrideFlags()
+	: bOverride_OpacityMaskClipValue(0)
+	, bOverride_BlendMode(0)
+	, bOverride_ShadingModel(0)
+	, bOverride_DitheredLODTransition(0)
+	, bOverride_CastDynamicShadowAsMasked(0)
+	, bOverride_TwoSided(0)
+	, bOverride_bIsThinSurface(0)
+	, bOverride_OutputTranslucentVelocity(0)
+	, bOverride_bHasPixelAnimation(0)
+	, bOverride_bEnableTessellation(0)
+	, Override_UsageFlags(0)
+{
+}
+
 void UMaterialValidationGroup::Serialize(FArchive& Ar)
 {
 	Super::Serialize(Ar);
 	Ar.UsingCustomVersion(FMaterialValidationGroupObjectVersion::GUID);
 
-	if (Ar.IsLoading() && Ar.CustomVer(FMaterialValidationGroupObjectVersion::GUID) < FMaterialValidationGroupObjectVersion::UpdatedPermutationHash2)
+	if (Ar.IsLoading() && Ar.CustomVer(FMaterialValidationGroupObjectVersion::GUID) < FMaterialValidationGroupObjectVersion::AddedAssetData)
 	{
 		// Clear invalidated data.
 		Materials.Reset();
@@ -76,17 +108,13 @@ void UMaterialValidationGroup::PostEditChangeProperty(FPropertyChangedEvent& Pro
 	}
 }
 
-void UMaterialValidationGroup::RemoveMissing()
-{
-	UMaterialValidationLibrary::RemoveMissingMaterialsFromGroup(this);
-}
-
-void UMaterialValidationGroup::UpdateFromPaths()
+void UMaterialValidationGroup::UpdateMaterials()
 {
-	UMaterialValidationLibrary::AddMaterialsFromGroupDirectoryPaths(this);
+	UMaterialValidationLibrary::RemoveInvalidMaterialsFromGroup(this);
+	UMaterialValidationLibrary::AddMissingMaterialsToGroup(this);
 }
 
-void UMaterialValidationGroup::Update()
+void UMaterialValidationGroup::UpdatePermutations()
 {
-	UMaterialValidationLibrary::UpdateAllMaterialsInGroup(this);
+	UMaterialValidationLibrary::UpdateMaterialPermutationsInGroup(this);
 }
diff --git a/Engine/Plugins/Editor/MaterialValidation/Source/Private/MaterialValidationGroup.h b/Engine/Plugins/Editor/MaterialValidation/Source/Private/MaterialValidationGroup.h
--- a/Engine/Plugins/Editor/MaterialValidation/Source/Private/MaterialValidationGroup.h
+++ b/Engine/Plugins/Editor/MaterialValidation/Source/Private/MaterialValidationGroup.h
@@ -14,13 +14,32 @@ struct FMaterialValidationDesc
 {
 	GENERATED_BODY()
 	
-	/** The approved base material permutation. */
+	/** The asset data for the material. */
 	UPROPERTY()
-	uint32 BaseMaterialHash;
+	FMaterialValidationAssetData_Material AssetData;
 
-	/** Array of currently approved material permutations not including the BaseMaterialHash. */
+	/** The non-default asset datas from all of the material instances in the material hierarchy. */
 	UPROPERTY()
-	TArray<uint32> ApprovedPermutationHashes;
+	TArray<FMaterialValidationAssetData_MaterialInstance> MaterialInstanceAssetDatas;
+
+	/** 
+	 * Map of UMaterialInstance asset paths to MaterialInstanceAssetDatas array index. Value of -1 indicates a default AssetData.
+	 * Paths are FString to avoid references to the actual assets which would need managing on asset delete etc. 
+	 */
+	UPROPERTY()
+	TMap<FString, int32> MaterialInstances;
+
+	/** Array of material permutations not including the one from the base material which is stored in AssetData.PermutationHash. */
+	UPROPERTY()
+	TArray<uint32> PermutationHashes;
+
+	/** Number of shader types on the base material permutation. */
+	UPROPERTY()
+	int32 NumShadersBase;
+
+	/** Total sum of shader types for all permutations found of this material. */
+	UPROPERTY(VisibleAnywhere, Category = Material)
+	int32 NumShadersTotal;
 };
 
 /** Data Asset containing the validation setup for a group of UMaterials. */
@@ -50,23 +69,19 @@ public:
 	UPROPERTY(VisibleAnywhere, Category = Material)
 	TMap<FString, FMaterialValidationDesc> Materials;
 
-	/** Remove all Material assets that cannot be found on disk. */
-	UFUNCTION(CallInEditor, Category = Actions, meta = (DisplayName = "Remove Missing Materials"))
-	void RemoveMissing();
-
 	/** 
 	 * Add all Material assets found under the material search paths.
-	 * The new material assets will have empty approved permutation lists.
+	 * The new material assets will have empty permutation lists.
 	 */
-	UFUNCTION(CallInEditor, Category = Actions, meta = (DisplayName = "Update From Material Paths"))
-	void UpdateFromPaths();
+	UFUNCTION(CallInEditor, Category = Actions)
+	void UpdateMaterials();
 
 	/** 
 	 * Update current state of all Materials.
-	 * For each material this will traverse the material instance hierarchy and add all existing permutations to the approved permutation list.
+	 * For each material this will traverse the material instance hierarchy and add all existing permutations to the permutation list.
 	 */
-	UFUNCTION(CallInEditor, Category = Actions, meta = (DisplayName = "Update Approved Lists"))
-	void Update();
+	UFUNCTION(CallInEditor, Category = Actions)
+	void UpdatePermutations();
 
 protected:
 	virtual void Serialize(FArchive& Ar) override;
diff --git a/Engine/Plugins/Editor/MaterialValidation/Source/Private/MaterialValidationLibrary.cpp b/Engine/Plugins/Editor/MaterialValidation/Source/Private/MaterialValidationLibrary.cpp
--- a/Engine/Plugins/Editor/MaterialValidation/Source/Private/MaterialValidationLibrary.cpp
+++ b/Engine/Plugins/Editor/MaterialValidation/Source/Private/MaterialValidationLibrary.cpp
@@ -6,10 +6,10 @@
 #include "AssetRegistry/AssetRegistryModule.h"
 #include "AssetRegistry/IAssetRegistry.h"
 #include "Interfaces/IPluginManager.h"
+#include "MaterialEditingLibrary.h"
 #include "Materials/Material.h"
 #include "Materials/MaterialInstanceConstant.h"
 #include "Materials/MaterialLayersFunctions.h"
-#include "MaterialShared.h"
 #include "MaterialValidationConfig.h"
 #include "MaterialValidationGroup.h"
 #include "MaterialValidationModule.h"
@@ -44,74 +44,106 @@ static void BuildPermutationData(UMaterial const& InBaseMaterial, FMaterialPermu
 	OutPermutationData.StaticPropertyLayoutHash = HashCombine(GetTypeHash(StaticSwitchGuids), GetTypeHash(ComponentMaskGuids));
 }
 
-/** 
- * Initialize the FMaterialInstanceBasePropertyOverrides values (but not override flags) from a material.
- * We only consider the parts of FMaterialInstanceBasePropertyOverrides that trigger shader permutations.
- * We could use a more compact struct with only these parts, but we just reuse a sparse part of this existing struct.
- */
-static FMaterialInstanceBasePropertyOverrides InitRelevantBaseProperties(UMaterialInterface const& InMaterial)
+/** Get the fully resolved static permutation values from a material or material instance. */
+static FStaticPermutationProperties GetStaticPermutationProperties(UMaterialInterface const& InMaterial)
+{
+	FStaticPermutationProperties StaticProperties;
+	StaticProperties.OpacityMaskClipValue = InMaterial.GetOpacityMaskClipValue();
+	StaticProperties.BlendMode = InMaterial.GetBlendMode();
+	StaticProperties.ShadingModel = InMaterial.GetShadingModels().GetFirstShadingModel();
+	StaticProperties.DitheredLODTransition = InMaterial.IsDitheredLODTransition();
+	StaticProperties.bCastDynamicShadowAsMasked = InMaterial.GetCastDynamicShadowAsMasked();
+	StaticProperties.TwoSided = InMaterial.IsTwoSided();
+	StaticProperties.bIsThinSurface = InMaterial.IsThinSurface();
+	StaticProperties.bOutputTranslucentVelocity = InMaterial.IsTranslucencyWritingVelocity();
+	StaticProperties.bHasPixelAnimation = InMaterial.HasPixelAnimation();
+	StaticProperties.bEnableTessellation = InMaterial.IsTessellationEnabled();
+	StaticProperties.UsageFlags = InMaterial.GetUsageFlags();
+	return StaticProperties;
+}
+
+/** Get the relevant static permutation values from a material instance. */
+static FStaticPermutationProperties GetStaticPermutationPropertyOverrides(UMaterialInstance const& InMaterialInstance)
+{
+	FStaticPermutationProperties StaticProperties;
+	StaticProperties.OpacityMaskClipValue = InMaterialInstance.BasePropertyOverrides.OpacityMaskClipValue;
+	StaticProperties.BlendMode = InMaterialInstance.BasePropertyOverrides.BlendMode;
+	StaticProperties.ShadingModel = InMaterialInstance.BasePropertyOverrides.ShadingModel;
+	StaticProperties.DitheredLODTransition = InMaterialInstance.BasePropertyOverrides.DitheredLODTransition;
+	StaticProperties.bCastDynamicShadowAsMasked = InMaterialInstance.BasePropertyOverrides.bCastDynamicShadowAsMasked;
+	StaticProperties.TwoSided = InMaterialInstance.BasePropertyOverrides.TwoSided;
+	StaticProperties.bIsThinSurface = InMaterialInstance.BasePropertyOverrides.bIsThinSurface;
+	StaticProperties.bOutputTranslucentVelocity = InMaterialInstance.BasePropertyOverrides.bOutputTranslucentVelocity;
+	StaticProperties.bHasPixelAnimation = InMaterialInstance.BasePropertyOverrides.bHasPixelAnimation;
+	StaticProperties.bEnableTessellation = InMaterialInstance.BasePropertyOverrides.bEnableTessellation;
+	StaticProperties.UsageFlags = InMaterialInstance.BasePropertyOverrides.UsageFlags;
+	return StaticProperties;
+}
+
+/** Get the relevant static permutation override flags for values are enabled from a material instance. */
+static FStaticPermutationPropertyOverrideFlags GetStaticPermutationPropertyOverrideFlags(UMaterialInstance const& InMaterialInstance)
+{
+	FStaticPermutationPropertyOverrideFlags OverrideFlags;
+	OverrideFlags.bOverride_OpacityMaskClipValue = InMaterialInstance.BasePropertyOverrides.bOverride_OpacityMaskClipValue;
+	OverrideFlags.bOverride_BlendMode = InMaterialInstance.BasePropertyOverrides.bOverride_BlendMode;
+	OverrideFlags.bOverride_ShadingModel = InMaterialInstance.BasePropertyOverrides.bOverride_ShadingModel;
+	OverrideFlags.bOverride_DitheredLODTransition = InMaterialInstance.BasePropertyOverrides.bOverride_DitheredLODTransition;
+	OverrideFlags.bOverride_CastDynamicShadowAsMasked = InMaterialInstance.BasePropertyOverrides.bOverride_CastDynamicShadowAsMasked;
+	OverrideFlags.bOverride_TwoSided = InMaterialInstance.BasePropertyOverrides.bOverride_TwoSided;
+	OverrideFlags.bOverride_bIsThinSurface = InMaterialInstance.BasePropertyOverrides.bOverride_bIsThinSurface;
+	OverrideFlags.bOverride_OutputTranslucentVelocity = InMaterialInstance.BasePropertyOverrides.bOverride_OutputTranslucentVelocity;
+	OverrideFlags.bOverride_bHasPixelAnimation = InMaterialInstance.BasePropertyOverrides.bOverride_bHasPixelAnimation;
+	OverrideFlags.bOverride_bEnableTessellation = InMaterialInstance.BasePropertyOverrides.bOverride_bEnableTessellation;
+	OverrideFlags.Override_UsageFlags = InMaterialInstance.BasePropertyOverrides.bOverride_UsageFlags;
+	return OverrideFlags;
+}
+
+/** Fill a FMaterialInstanceBasePropertyOverrides from FStaticPermutationProperties, setting all the override flags. */
+FMaterialInstanceBasePropertyOverrides GetStaticPropertiesAsBasePropertyOverrides(FStaticPermutationProperties const& InStaticProperties)
 {
 	FMaterialInstanceBasePropertyOverrides BaseProperties;
-	BaseProperties.OpacityMaskClipValue = InMaterial.GetOpacityMaskClipValue();
-	BaseProperties.BlendMode = InMaterial.GetBlendMode();
-	BaseProperties.ShadingModel = InMaterial.GetShadingModels().GetFirstShadingModel();
-	BaseProperties.DitheredLODTransition = InMaterial.IsDitheredLODTransition();
-	BaseProperties.bCastDynamicShadowAsMasked = InMaterial.GetCastDynamicShadowAsMasked();
-	BaseProperties.TwoSided = InMaterial.IsTwoSided();
-	BaseProperties.bIsThinSurface = InMaterial.IsThinSurface();
-	BaseProperties.bOutputTranslucentVelocity = InMaterial.IsTranslucencyWritingVelocity();
-	BaseProperties.bHasPixelAnimation = InMaterial.HasPixelAnimation();
-	BaseProperties.bEnableTessellation = InMaterial.IsTessellationEnabled();
-	BaseProperties.UsageFlags = InMaterial.GetUsageFlags();
+	BaseProperties.bOverride_OpacityMaskClipValue = true;
+	BaseProperties.OpacityMaskClipValue = InStaticProperties.OpacityMaskClipValue;
+	BaseProperties.bOverride_BlendMode = true;
+	BaseProperties.BlendMode = InStaticProperties.BlendMode;
+	BaseProperties.bOverride_ShadingModel = true;
+	BaseProperties.ShadingModel = InStaticProperties.ShadingModel;
+	BaseProperties.bOverride_DitheredLODTransition = true;
+	BaseProperties.DitheredLODTransition = InStaticProperties.DitheredLODTransition;
+	BaseProperties.bOverride_CastDynamicShadowAsMasked = true;
+	BaseProperties.bCastDynamicShadowAsMasked = InStaticProperties.bCastDynamicShadowAsMasked;
+	BaseProperties.bOverride_TwoSided = true;
+	BaseProperties.TwoSided = InStaticProperties.TwoSided;
+	BaseProperties.bOverride_bIsThinSurface = true;
+	BaseProperties.bIsThinSurface = InStaticProperties.bIsThinSurface;
+	BaseProperties.bOverride_OutputTranslucentVelocity = true;
+	BaseProperties.bOutputTranslucentVelocity = InStaticProperties.bOutputTranslucentVelocity;
+	BaseProperties.bOverride_bHasPixelAnimation = true;
+	BaseProperties.bHasPixelAnimation = InStaticProperties.bHasPixelAnimation;
+	BaseProperties.bOverride_bEnableTessellation = true;
+	BaseProperties.bEnableTessellation = InStaticProperties.bEnableTessellation;
+	BaseProperties.bOverride_UsageFlags = ~0u;
+	BaseProperties.UsageFlags = InStaticProperties.UsageFlags;
 	return BaseProperties;
 }
 
-/** Copy the relevant parts of the FMaterialInstanceBasePropertyOverrides struct from the value stored on the Material Instance. */
-static FMaterialInstanceBasePropertyOverrides InitRelevantBasePropertyOverrides(UMaterialInstance const& InMaterialInstance)
-{
-	FMaterialInstanceBasePropertyOverrides BasePropertyOverrides;
-	BasePropertyOverrides.bOverride_OpacityMaskClipValue = InMaterialInstance.BasePropertyOverrides.bOverride_OpacityMaskClipValue;
-	BasePropertyOverrides.OpacityMaskClipValue = InMaterialInstance.BasePropertyOverrides.OpacityMaskClipValue;
-	BasePropertyOverrides.bOverride_BlendMode = InMaterialInstance.BasePropertyOverrides.bOverride_BlendMode;
-	BasePropertyOverrides.BlendMode = InMaterialInstance.BasePropertyOverrides.BlendMode;
-	BasePropertyOverrides.bOverride_ShadingModel = InMaterialInstance.BasePropertyOverrides.bOverride_ShadingModel;
-	BasePropertyOverrides.ShadingModel = InMaterialInstance.BasePropertyOverrides.ShadingModel;
-	BasePropertyOverrides.bOverride_DitheredLODTransition = InMaterialInstance.BasePropertyOverrides.bOverride_DitheredLODTransition;
-	BasePropertyOverrides.DitheredLODTransition = InMaterialInstance.BasePropertyOverrides.DitheredLODTransition;
-	BasePropertyOverrides.bCastDynamicShadowAsMasked = InMaterialInstance.BasePropertyOverrides.bCastDynamicShadowAsMasked;
-	BasePropertyOverrides.bOverride_CastDynamicShadowAsMasked = InMaterialInstance.BasePropertyOverrides.bOverride_CastDynamicShadowAsMasked;
-	BasePropertyOverrides.bOverride_TwoSided = InMaterialInstance.BasePropertyOverrides.bOverride_TwoSided;
-	BasePropertyOverrides.TwoSided = InMaterialInstance.BasePropertyOverrides.TwoSided;
-	BasePropertyOverrides.bOverride_bIsThinSurface = InMaterialInstance.BasePropertyOverrides.bOverride_bIsThinSurface;
-	BasePropertyOverrides.bIsThinSurface = InMaterialInstance.BasePropertyOverrides.bIsThinSurface;
-	BasePropertyOverrides.bOverride_OutputTranslucentVelocity = InMaterialInstance.BasePropertyOverrides.bOverride_OutputTranslucentVelocity;
-	BasePropertyOverrides.bOutputTranslucentVelocity = InMaterialInstance.BasePropertyOverrides.bOutputTranslucentVelocity;
-	BasePropertyOverrides.bOverride_bHasPixelAnimation = InMaterialInstance.BasePropertyOverrides.bOverride_bHasPixelAnimation;
-	BasePropertyOverrides.bHasPixelAnimation = InMaterialInstance.BasePropertyOverrides.bHasPixelAnimation;
-	BasePropertyOverrides.bOverride_bEnableTessellation = InMaterialInstance.BasePropertyOverrides.bOverride_bEnableTessellation;
-	BasePropertyOverrides.bEnableTessellation = InMaterialInstance.BasePropertyOverrides.bEnableTessellation;
-	BasePropertyOverrides.bOverride_UsageFlags = InMaterialInstance.BasePropertyOverrides.bOverride_UsageFlags;
-	BasePropertyOverrides.UsageFlags = InMaterialInstance.BasePropertyOverrides.UsageFlags;
-	return BasePropertyOverrides;
-}
-
-/** Return a hash built from the relevant FMaterialInstanceBasePropertyOverrides values (but not override flags). */
-static uint32 BuildBasePropertiesHash(FMaterialInstanceBasePropertyOverrides const& InBaseProperties)
+/** Return a hash built from the FStaticPermutationProperties values. */
+static uint32 BuildStaticPermutationPropertiesHash(FStaticPermutationProperties const& InStaticProperties)
 {
 	const uint8 PackedBools = 
-		(InBaseProperties.TwoSided ? 1 << 0 : 0) |
-		(InBaseProperties.bIsThinSurface ? 1 << 1 : 0) |
-		(InBaseProperties.DitheredLODTransition ? 1 << 2 : 0) |
-		(InBaseProperties.bCastDynamicShadowAsMasked ? 1 << 3 : 0) |
-		(InBaseProperties.bOutputTranslucentVelocity ? 1 << 4 : 0) |
-		(InBaseProperties.bHasPixelAnimation ? 1 << 5 : 0) |
-		(InBaseProperties.bEnableTessellation ? 1 << 6 : 0);
+		(InStaticProperties.TwoSided ? 1 << 0 : 0) |
+		(InStaticProperties.bIsThinSurface ? 1 << 1 : 0) |
+		(InStaticProperties.DitheredLODTransition ? 1 << 2 : 0) |
+		(InStaticProperties.bCastDynamicShadowAsMasked ? 1 << 3 : 0) |
+		(InStaticProperties.bOutputTranslucentVelocity ? 1 << 4 : 0) |
+		(InStaticProperties.bHasPixelAnimation ? 1 << 5 : 0) |
+		(InStaticProperties.bEnableTessellation ? 1 << 6 : 0);
 
 	uint32 Hash = GetTypeHash(PackedBools);
-	Hash = HashCombineFast(Hash, GetTypeHash(InBaseProperties.OpacityMaskClipValue));
-	Hash = HashCombineFast(Hash, GetTypeHash(InBaseProperties.BlendMode));
-	Hash = HashCombineFast(Hash, GetTypeHash(InBaseProperties.ShadingModel));
-	Hash = HashCombineFast(Hash, GetTypeHash(InBaseProperties.UsageFlags));
+	Hash = HashCombineFast(Hash, GetTypeHash(InStaticProperties.OpacityMaskClipValue));
+	Hash = HashCombineFast(Hash, GetTypeHash(InStaticProperties.BlendMode));
+	Hash = HashCombineFast(Hash, GetTypeHash(InStaticProperties.ShadingModel));
+	Hash = HashCombineFast(Hash, GetTypeHash(InStaticProperties.UsageFlags));
 	return Hash;
 }
 
@@ -127,9 +159,9 @@ static uint32 BuildLayerHash(UMaterialInterface const& InMaterial)
 }
 
 /** Return a hash built for a material or material instance which will match for all instances that generate the same shader permutations. */
-static uint32 BuildPermutationHash(FMaterialInstanceBasePropertyOverrides const& InBaseProperties, TConstArrayView<uint32> InStaticSwitchValues, TConstArrayView<uint32> InComponentMaskValues, uint32 InMaterialLayerHash)
+static uint32 BuildPermutationHash(FStaticPermutationProperties const& InStaticProperties, TConstArrayView<uint32> InStaticSwitchValues, TConstArrayView<uint32> InComponentMaskValues, uint32 InMaterialLayerHash)
 {
-	uint32 Hash = BuildBasePropertiesHash(InBaseProperties);
+	uint32 Hash = BuildStaticPermutationPropertiesHash(InStaticProperties);
 	for (uint32 Value : InStaticSwitchValues)
 	{
 		Hash = HashCombineFast(Hash, GetTypeHash(Value));
@@ -142,28 +174,6 @@ static uint32 BuildPermutationHash(FMaterialInstanceBasePropertyOverrides const&
 	return Hash;
 }
 
-/** 
- * Get the (non-deduplicated) number of shaders that a single material will generat. 
- * Note: This is only calculated for the current (editor) platform and probably contains debug shaders that aren't cooked.
- */
-static int32 GetShaderCount(UMaterialInterface const& InMaterial)
-{
-	// todo: GetShaderTypes() isn't const, but maybe could be?
-	TArray<FDebugShaderTypeInfo> ShaderInfos;
-	(const_cast<UMaterialInterface&>(InMaterial)).GetShaderTypes(GMaxRHIShaderPlatform, nullptr, ShaderInfos);
-
-	int NumShaders = 0;
-	for (FDebugShaderTypeInfo const& ShaderInfo : ShaderInfos)
-	{
-		NumShaders += ShaderInfo.ShaderTypes.Num();
-		for (FDebugShaderPipelineInfo const& PipelineInfo : ShaderInfo.Pipelines)
-		{
-			NumShaders += PipelineInfo.ShaderTypes.Num();
-		}
-	}
-	return NumShaders;
-}
-
 /** Extract AssetData from a UMaterial object. */
 static void GetAssetData(
 	UMaterial const& InMaterial, 
@@ -174,7 +184,7 @@ static void GetAssetData(
 	{
 		const uint32 StaticSwitchNum = InPermutationData.StaticSwitchParameterInfo.Num();
 		const uint32 StaticSwitchBitArrayDwordNum = (StaticSwitchNum + 31) / 32;
-		OutAssetData.StaticSwitchValues.AddZeroed(StaticSwitchBitArrayDwordNum);
+		OutAssetData.StaticSwitchValues.SetNumZeroed(StaticSwitchBitArrayDwordNum);
 
 		for (uint32 Index = 0; Index < StaticSwitchNum; ++Index)
 		{
@@ -194,7 +204,7 @@ static void GetAssetData(
 	{
 		const uint32 ComponentMaskNum = InPermutationData.ComponentMaskParameterInfo.Num();
 		const uint32 ComponentMaskBitArrayDwordNum = ((ComponentMaskNum * 4) + 31) / 32;
-		OutAssetData.ComponentMaskValues.AddZeroed(ComponentMaskBitArrayDwordNum);
+		OutAssetData.ComponentMaskValues.SetNumZeroed(ComponentMaskBitArrayDwordNum);
 
 		for (uint32 Index = 0; Index < ComponentMaskNum; ++Index)
 		{
@@ -212,22 +222,22 @@ static void GetAssetData(
 		}
 	}
 
-	OutAssetData.BaseProperties = InitRelevantBaseProperties(InMaterial);
+	OutAssetData.StaticProperties = GetStaticPermutationProperties(InMaterial);
 
 	OutAssetData.MaterialLayerHash = BuildLayerHash(InMaterial);
 	
 	OutAssetData.StaticPropertyLayoutHash = InPermutationData.StaticPropertyLayoutHash;
-	OutAssetData.PermutationHash = BuildPermutationHash(OutAssetData.BaseProperties, OutAssetData.StaticSwitchValues, OutAssetData.ComponentMaskValues, OutAssetData.MaterialLayerHash);
-	OutAssetData.ShaderCount = GetShaderCount(InMaterial);
+	OutAssetData.PermutationHash = BuildPermutationHash(OutAssetData.StaticProperties, OutAssetData.StaticSwitchValues, OutAssetData.ComponentMaskValues, OutAssetData.MaterialLayerHash);
 }
 
 /** Extract AssetData from a UMaterialInstance object. */
 static void GetAssetData(
 	UMaterialInstance const& InMaterialInstance, 
-	UMaterialInterface const& InParentMaterial, 
 	FMaterialPermutationData const& InPermutationData,
 	FMaterialValidationAssetData_MaterialInstance& OutAssetData)
 {
+	UMaterialInterface const* ParentMaterial = InMaterialInstance.Parent;
+
 	// We collect the static switch and component mask values to create the
 	// permutation hash but only store the overridden values to the asset data.
 	TArray<uint32> StaticSwitchValues;
@@ -237,9 +247,9 @@ static void GetAssetData(
 	{
 		const uint32 StaticSwitchNum = InPermutationData.StaticSwitchParameterInfo.Num();
 		const uint32 StaticSwitchBitArrayDwordNum = (StaticSwitchNum + 31) / 32;
-		OutAssetData.StaticSwitchOverrideMask.AddZeroed(StaticSwitchBitArrayDwordNum);
-		OutAssetData.StaticSwitchOverrideValues.AddZeroed(StaticSwitchBitArrayDwordNum);
-		StaticSwitchValues.AddZeroed(StaticSwitchBitArrayDwordNum);
+		OutAssetData.StaticSwitchOverrideMask.SetNumZeroed(StaticSwitchBitArrayDwordNum);
+		OutAssetData.StaticSwitchOverrideValues.SetNumZeroed(StaticSwitchBitArrayDwordNum);
+		StaticSwitchValues.SetNumZeroed(StaticSwitchBitArrayDwordNum);
 
 		for (uint32 Index = 0; Index < StaticSwitchNum; ++Index)
 		{
@@ -248,7 +258,7 @@ static void GetAssetData(
 			const bool Value = bInstanceValid && InstanceMeta.Value.AsStaticSwitch();
 
 			FMaterialParameterMetadata ParentMeta;
-			const bool bBaseValid = InParentMaterial.GetParameterValue(EMaterialParameterType::StaticSwitch, InPermutationData.StaticSwitchParameterInfo[Index], ParentMeta);
+			const bool bBaseValid = ParentMaterial->GetParameterValue(EMaterialParameterType::StaticSwitch, InPermutationData.StaticSwitchParameterInfo[Index], ParentMeta);
 			const bool ParentValue = bBaseValid && ParentMeta.Value.AsStaticSwitch();
 
 			const uint32 BitArrayIndex = Index / 32;
@@ -270,9 +280,9 @@ static void GetAssetData(
 	{
 		const uint32 ComponentMaskNum = InPermutationData.ComponentMaskParameterInfo.Num();
 		const uint32 ComponentMaskBitValueArrayDwordNum = ((ComponentMaskNum * 4) + 31) / 32;
-		OutAssetData.ComponentMaskOverrideValues.AddZeroed(ComponentMaskBitValueArrayDwordNum);
-		OutAssetData.ComponentMaskOverrideMask.AddZeroed(ComponentMaskBitValueArrayDwordNum);
-		ComponentMaskValues.AddZeroed(ComponentMaskBitValueArrayDwordNum);
+		OutAssetData.ComponentMaskOverrideValues.SetNumZeroed(ComponentMaskBitValueArrayDwordNum);
+		OutAssetData.ComponentMaskOverrideMask.SetNumZeroed(ComponentMaskBitValueArrayDwordNum);
+		ComponentMaskValues.SetNumZeroed(ComponentMaskBitValueArrayDwordNum);
 
 		for (uint32 Index = 0; Index < ComponentMaskNum; ++Index)
 		{
@@ -282,7 +292,7 @@ static void GetAssetData(
 			const uint32 ValueMask = (Value.R ? 1 : 0) | (Value.G ? 2 : 0) | (Value.B ? 4 : 0) | (Value.A ? 8 : 0);
 
 			FMaterialParameterMetadata ParentMeta;
-			const bool bParentValid = InParentMaterial.GetParameterValue(EMaterialParameterType::StaticComponentMask, InPermutationData.ComponentMaskParameterInfo[Index], ParentMeta);
+			const bool bParentValid = ParentMaterial->GetParameterValue(EMaterialParameterType::StaticComponentMask, InPermutationData.ComponentMaskParameterInfo[Index], ParentMeta);
 			const FStaticComponentMaskValue ParentValue = bParentValid ? ParentMeta.Value.AsStaticComponentMask() : FStaticComponentMaskValue();
 
 			const uint32 BitArrayIndex = Index / 8;
@@ -300,13 +310,13 @@ static void GetAssetData(
 		}
 	}
 
-	FMaterialInstanceBasePropertyOverrides BaseProperties = InitRelevantBaseProperties(InMaterialInstance);
-	OutAssetData.BasePropertyOverrides = InitRelevantBasePropertyOverrides(InMaterialInstance);
+	FStaticPermutationProperties StaticProperties = GetStaticPermutationProperties(InMaterialInstance);
+	OutAssetData.StaticProperties = GetStaticPermutationPropertyOverrides(InMaterialInstance);
+	OutAssetData.StaticPropertyOverrideFlags = GetStaticPermutationPropertyOverrideFlags(InMaterialInstance);
 
 	OutAssetData.MaterialLayerHash = BuildLayerHash(InMaterialInstance);
 	
-	OutAssetData.PermutationHash = BuildPermutationHash(BaseProperties, StaticSwitchValues, ComponentMaskValues, OutAssetData.MaterialLayerHash);
-	OutAssetData.ShaderCount = GetShaderCount(InMaterialInstance);
+	OutAssetData.PermutationHash = BuildPermutationHash(StaticProperties, StaticSwitchValues, ComponentMaskValues, OutAssetData.MaterialLayerHash);
 }
 
 /** Build hash for a material instance which will match for all instances that generate the same shader permutations. */
@@ -316,16 +326,17 @@ static uint32 BuildPermutationHash(UMaterialInstanceConstant const& InMaterialIn
 	BuildPermutationData(InBaseMaterial, PermutationData);
 
 	FMaterialValidationAssetData_MaterialInstance InstanceAssetData;
-	GetAssetData(InMaterialInstance, *InMaterialInstance.Parent, PermutationData, InstanceAssetData);
+	GetAssetData(InMaterialInstance, PermutationData, InstanceAssetData);
 
 	return InstanceAssetData.PermutationHash;
 }
 
 /** Apply overrides for bit mask arrays such as those used by static switch and component masks. */
-static void ApplyBitMaskOverrides(TConstArrayView<uint32> InParentValues, TConstArrayView<uint32> InChildOverridesMasks, TConstArrayView<uint32> InChildOverrideValues, TArrayView<uint32> OutValues)
+static void ApplyBitMaskOverrides(TConstArrayView<uint32> InParentValues, TConstArrayView<uint32> InChildOverridesMasks, TConstArrayView<uint32> InChildOverrideValues, TArray<uint32>& OutValues)
 {
 	const int32 NumDwords = InParentValues.Num();
-	check(NumDwords == InChildOverridesMasks.Num() && NumDwords == InChildOverrideValues.Num() && NumDwords == OutValues.Num());
+	check(NumDwords == InChildOverridesMasks.Num() && NumDwords == InChildOverrideValues.Num());
+	OutValues.SetNumUninitialized(NumDwords);
 
 	for (int32 Index = 0; Index < NumDwords; ++Index)
 	{
@@ -334,19 +345,23 @@ static void ApplyBitMaskOverrides(TConstArrayView<uint32> InParentValues, TConst
 }
 
 /** Apply material base property overrides to generate the child base property values from the parent values and the child overrides. */
-static void ApplyBasePropertyOverrides(FMaterialInstanceBasePropertyOverrides const& InParentBaseProperties, FMaterialInstanceBasePropertyOverrides const& InChildBasePropertyOverrides, FMaterialInstanceBasePropertyOverrides& OutBaseProperties)
+static void ApplyBasePropertyOverrides(
+	FStaticPermutationProperties const& InParentProperties, 
+	FStaticPermutationPropertyOverrideFlags const& InChildOverrideFlags,
+	FStaticPermutationProperties const& InChildOverrideProperties,
+	FStaticPermutationProperties& OutProperties)
 {
-	OutBaseProperties.TwoSided = InChildBasePropertyOverrides.bOverride_TwoSided ? InChildBasePropertyOverrides.TwoSided : InParentBaseProperties.TwoSided;
-	OutBaseProperties.bIsThinSurface = InChildBasePropertyOverrides.bOverride_bIsThinSurface ? InChildBasePropertyOverrides.bIsThinSurface : InParentBaseProperties.bIsThinSurface;
-	OutBaseProperties.DitheredLODTransition = InChildBasePropertyOverrides.bOverride_DitheredLODTransition ? InChildBasePropertyOverrides.DitheredLODTransition : InParentBaseProperties.DitheredLODTransition;
-	OutBaseProperties.bCastDynamicShadowAsMasked = InChildBasePropertyOverrides.bOverride_CastDynamicShadowAsMasked ? InChildBasePropertyOverrides.bCastDynamicShadowAsMasked : InParentBaseProperties.bCastDynamicShadowAsMasked;
-	OutBaseProperties.bOutputTranslucentVelocity = InChildBasePropertyOverrides.bOverride_OutputTranslucentVelocity ? InChildBasePropertyOverrides.bOutputTranslucentVelocity : InParentBaseProperties.bOutputTranslucentVelocity;
-	OutBaseProperties.bHasPixelAnimation = InChildBasePropertyOverrides.bOverride_bHasPixelAnimation ? InChildBasePropertyOverrides.bHasPixelAnimation : InParentBaseProperties.bHasPixelAnimation;
-	OutBaseProperties.bEnableTessellation = InChildBasePropertyOverrides.bOverride_bEnableTessellation ? InChildBasePropertyOverrides.bEnableTessellation : InParentBaseProperties.bEnableTessellation;
-	OutBaseProperties.BlendMode = InChildBasePropertyOverrides.bOverride_BlendMode ? InChildBasePropertyOverrides.BlendMode : InParentBaseProperties.BlendMode;
-	OutBaseProperties.ShadingModel = InChildBasePropertyOverrides.bOverride_ShadingModel ? InChildBasePropertyOverrides.ShadingModel : InParentBaseProperties.ShadingModel;
-	OutBaseProperties.OpacityMaskClipValue = InChildBasePropertyOverrides.bOverride_OpacityMaskClipValue ? InChildBasePropertyOverrides.OpacityMaskClipValue : InParentBaseProperties.OpacityMaskClipValue;
-	OutBaseProperties.UsageFlags = (InParentBaseProperties.UsageFlags & ~InChildBasePropertyOverrides.bOverride_UsageFlags) | (InChildBasePropertyOverrides.bOverride_UsageFlags & InChildBasePropertyOverrides.UsageFlags);
+	OutProperties.TwoSided = InChildOverrideFlags.bOverride_TwoSided ? InChildOverrideProperties.TwoSided : InParentProperties.TwoSided;
+	OutProperties.bIsThinSurface = InChildOverrideFlags.bOverride_bIsThinSurface ? InChildOverrideProperties.bIsThinSurface : InParentProperties.bIsThinSurface;
+	OutProperties.DitheredLODTransition = InChildOverrideFlags.bOverride_DitheredLODTransition ? InChildOverrideProperties.DitheredLODTransition : InParentProperties.DitheredLODTransition;
+	OutProperties.bCastDynamicShadowAsMasked = InChildOverrideFlags.bOverride_CastDynamicShadowAsMasked ? InChildOverrideProperties.bCastDynamicShadowAsMasked : InParentProperties.bCastDynamicShadowAsMasked;
+	OutProperties.bOutputTranslucentVelocity = InChildOverrideFlags.bOverride_OutputTranslucentVelocity ? InChildOverrideProperties.bOutputTranslucentVelocity : InParentProperties.bOutputTranslucentVelocity;
+	OutProperties.bHasPixelAnimation = InChildOverrideFlags.bOverride_bHasPixelAnimation ? InChildOverrideProperties.bHasPixelAnimation : InParentProperties.bHasPixelAnimation;
+	OutProperties.bEnableTessellation = InChildOverrideFlags.bOverride_bEnableTessellation ? InChildOverrideProperties.bEnableTessellation : InParentProperties.bEnableTessellation;
+	OutProperties.BlendMode = InChildOverrideFlags.bOverride_BlendMode ? InChildOverrideProperties.BlendMode : InParentProperties.BlendMode;
+	OutProperties.ShadingModel = InChildOverrideFlags.bOverride_ShadingModel ? InChildOverrideProperties.ShadingModel : InParentProperties.ShadingModel;
+	OutProperties.OpacityMaskClipValue = InChildOverrideFlags.bOverride_OpacityMaskClipValue ? InChildOverrideProperties.OpacityMaskClipValue : InParentProperties.OpacityMaskClipValue;
+	OutProperties.UsageFlags = (InParentProperties.UsageFlags & ~InChildOverrideFlags.Override_UsageFlags) | (InChildOverrideFlags.Override_UsageFlags & InChildOverrideProperties.UsageFlags);
 }
 
 /** Gather all search and exclude directory paths for a group. */
@@ -474,10 +489,39 @@ static FSoftObjectPath ConvertStringPathToSoftObjectPath(FString const& InAssetP
 }
 
 /** 
- * Load the instance hierarchy for a material and add all permutations to the FMaterialValidationDesc. 
- * This is potentially very slow for a material with many references that may also load/compile textures etc.
+ * Add the package path extension to a FName package path that has a missing extension. 
+ * todo: Ideally we would avoid this by only passing around package paths in one canonical form.
+ */
+FString AddPackagePathExtension(FName InPackagePath)
+{
+	FString PackagePath = InPackagePath.ToString();
+
+	int32 LastSlashIndex;
+	if (!PackagePath.FindLastChar('/', LastSlashIndex))
+	{
+		LastSlashIndex = -1;
+	}
+
+	int32 LastPeriodIndex;
+	if (PackagePath.FindLastChar('.', LastPeriodIndex) && LastPeriodIndex > LastSlashIndex)
+	{
+		return PackagePath;
+	}
+
+	FString AssetName = PackagePath.Mid(LastSlashIndex + 1);
+	if (AssetName.IsEmpty())
+	{
+		return PackagePath;
+	}
+	
+	return FString::Printf(TEXT("%s.%s"), *PackagePath, *AssetName);
+}
+
+/** 
+ * Traverse the instance hierarchy for a material and write to a depth first sorted array of package names.
+ * The array will have the input material package name as its first entry.
  */
-static void LoadMaterialInstanceConstants(UMaterial const& InMaterial, TArray<UMaterialInstanceConstant*>& OutMaterialInstances)
+static void GetMaterialInstanceHierarchy(UMaterial const& InMaterial, TArray<FName>& OutPackageNames, TArray<int32>& OutHierarchyDepths)
 {
 	FAssetRegistryModule& AssetRegistryModule = FModuleManager::LoadModuleChecked<FAssetRegistryModule>("AssetRegistry");
 	IAssetRegistry& AssetRegistry = AssetRegistryModule.Get();
@@ -488,19 +532,26 @@ static void LoadMaterialInstanceConstants(UMaterial const& InMaterial, TArray<UM
 		return;
 	}
 
-	TQueue<FName> ToVisit;
-	ToVisit.Enqueue(MaterialAssetData.PackageName);
-	TSet<FName> VisitedPackages;
+	// Maintain a stack of names to output.
+	TArray<FName> NameStack;
+	NameStack.Push(MaterialAssetData.PackageName);
+	// Maintain a stack of depths so that we can output those to OutHierarchyDepths.
+	TArray<int32> DepthStack;
+	DepthStack.Push(0);
 
-	while (!ToVisit.IsEmpty())
+	while (!NameStack.IsEmpty())
 	{
-		FName CurrentPackageName;
-		ToVisit.Dequeue(CurrentPackageName);
+		// Pop from stack and add to output.
+		const FName ParentPackageName = NameStack.Pop(EAllowShrinking::No);
+		OutPackageNames.Add(ParentPackageName);
+		const int32 ParentPackageDepth = DepthStack.Pop(EAllowShrinking::No);
+		OutHierarchyDepths.Add(ParentPackageDepth);
 
-		TArray<FName> ReferencerPackageNames;
-		AssetRegistry.GetReferencers(CurrentPackageName, ReferencerPackageNames);
+		// Push children onto our stack.
+		TArray<FName> ChildPackageNames;
+		AssetRegistry.GetReferencers(ParentPackageName, ChildPackageNames);
 
-		for (const FName& PackageName : ReferencerPackageNames)
+		for (const FName& PackageName : ChildPackageNames)
 		{
 			TArray<FAssetData> AssetDatas;
 			AssetRegistry.GetAssetsByPackageName(PackageName, AssetDatas, true);
@@ -515,28 +566,42 @@ static void LoadMaterialInstanceConstants(UMaterial const& InMaterial, TArray<UM
 				{
 					continue;
 				}
-				if (VisitedPackages.Contains(PackageName))
+				if (OutPackageNames.Contains(PackageName))
 				{
 					continue;
 				}
 
-				VisitedPackages.Add(PackageName);
-				ToVisit.Enqueue(PackageName);
+				NameStack.Push(PackageName);
+				DepthStack.Push(ParentPackageDepth + 1);
 				break;
 			}
 		}
 	}
 
-	UE_LOG(LogMaterialValidation, Display, TEXT("Found %d material instances for %s."), VisitedPackages.Num(), *InMaterial.GetName());
+	UE_LOG(LogMaterialValidation, Display, TEXT("Found %d material instances for %s."), OutPackageNames.Num() - 1, *InMaterial.GetName());
+}
 
-	FScopedSlowTask SlowTask(static_cast<float>(VisitedPackages.Num()), LOCTEXT("LoadingMaterialInstances", "Loading Material Instances..."));
+/**
+ * Load the instances for a material from an array of package names.
+ * Returns a material instance object array with the same indexing (with nullptr if we can't load).
+ * Loading the assets can be very slow for a material with many references that may also load/compile textures etc.
+ */
+static void LoadMaterialInstances(UMaterial const& InMaterial, TConstArrayView<FName> InPackageNames, TArray<UMaterialInstanceConstant*>& OutMaterialInstances)
+{
+	FAssetRegistryModule& AssetRegistryModule = FModuleManager::LoadModuleChecked<FAssetRegistryModule>("AssetRegistry");
+	IAssetRegistry& AssetRegistry = AssetRegistryModule.Get();
+
+	FScopedSlowTask SlowTask(static_cast<float>(InPackageNames.Num()), LOCTEXT("LoadingMaterialInstances", "Loading Material Instances..."));
 
-	for (FName PackageName : VisitedPackages)
+	OutMaterialInstances.SetNumZeroed(InPackageNames.Num());
+
+	for (int32 PackageIndex = 0; PackageIndex < InPackageNames.Num(); ++PackageIndex)
 	{
 		SlowTask.EnterProgressFrame();
 
+		// Each package name may map to several FAssetData so iterate them all to find one that has our base material.
 		TArray<FAssetData> AssetDatas;
-		AssetRegistry.GetAssetsByPackageName(PackageName, AssetDatas, true);
+		AssetRegistry.GetAssetsByPackageName(InPackageNames[PackageIndex], AssetDatas, true);
 		for (FAssetData const& AssetData : AssetDatas)
 		{
 			if (!AssetData.IsValid())
@@ -551,79 +616,221 @@ static void LoadMaterialInstanceConstants(UMaterial const& InMaterial, TArray<UM
 			{
 				if (MaterialInstance->GetMaterial() == &InMaterial)
 				{
-					OutMaterialInstances.Add(MaterialInstance);
+					OutMaterialInstances[PackageIndex] = MaterialInstance;
+					break;
 				}
 			}
 		}
 	}
 }
 
-/** Load the instance hierarchy for a material and add all permutations to the FMaterialValidationDesc. */
-static void UpdateValidMaterialInstances(UMaterial const& InMaterial, FMaterialValidationDesc& OutDesc)
+/**
+ * Get the (non-deduplicated) number of shaders that a single material or material instance will generate.
+ * Note: This is only calculated for the current (editor) platform and probably contains debug shaders that aren't cooked.
+ */
+static int32 GetNumShaderTypes(UMaterialInterface const& InMaterial)
 {
+	// todo: GetNumShaderTypes() isn't const, but maybe could be?
+	return UMaterialEditingLibrary::GetNumShaderTypes(const_cast<UMaterialInterface*>(&InMaterial));
+}
+
+/** 
+ * Update the FMaterialValidationDesc for a material.
+ * This contains asset data for all material instances and information about their unique permutations.
+ * Internally this calls LoadMaterialInstances() that is a potentially slow function to load all the material instance assets.
+ */
+static void UpdateValidMaterialInstances(UMaterial const& InMaterial, FMaterialValidationDesc& OutMaterialDesc)
+{
+	OutMaterialDesc.MaterialInstanceAssetDatas.Reset();
+	OutMaterialDesc.MaterialInstances.Reset();
+	OutMaterialDesc.PermutationHashes.Reset();
+
+	// Initialize shader counts.
+	OutMaterialDesc.NumShadersBase = GetNumShaderTypes(InMaterial);
+	OutMaterialDesc.NumShadersTotal = OutMaterialDesc.NumShadersBase;
+
+	// Get all material instance packages.
+	// PackageNames will contain InMaterial in the first entry.
+	TArray<FName> PackageNames;
+	TArray<int32> HierarchyDepths;
+	GetMaterialInstanceHierarchy(InMaterial, PackageNames, HierarchyDepths);
+
+	if (PackageNames.Num() <= 1)
+	{
+		return;
+	}
+
+	// Load the material instance assets.
 	TArray<UMaterialInstanceConstant*> MaterialInstances;
-	LoadMaterialInstanceConstants(InMaterial, MaterialInstances);
+	LoadMaterialInstances(InMaterial, PackageNames, MaterialInstances);
 
+	// Get the material asset data.
 	FMaterialPermutationData PermutationData;
 	BuildPermutationData(InMaterial, PermutationData);
+	GetAssetData(InMaterial, PermutationData, OutMaterialDesc.AssetData);
 
-	FMaterialValidationAssetData_Material AssetData;
-	GetAssetData(InMaterial, PermutationData, AssetData);
-	OutDesc.BaseMaterialHash = AssetData.PermutationHash;
-
-	OutDesc.ApprovedPermutationHashes.Reset();
+	// Walk the material instance hierarchy.
+	// We keep a stack of permutation hashes to detect when a child has the same hash as the parent.
+	TArray<uint32> PermutationHashStack;
+	PermutationHashStack.Push(OutMaterialDesc.AssetData.PermutationHash);
 
-	for (UMaterialInstanceConstant* MaterialInstance : MaterialInstances)
+	for (int32 Index = 1; Index < MaterialInstances.Num(); ++Index)
 	{
-		FMaterialValidationAssetData_MaterialInstance InstanceAssetData;
-		GetAssetData(*MaterialInstance, *MaterialInstance->Parent, PermutationData, InstanceAssetData);
+		const int32 CurrentDepth = HierarchyDepths[Index];
+		PermutationHashStack.SetNumZeroed(CurrentDepth + 1);
 
-		const uint32 Hash = InstanceAssetData.PermutationHash;
-		if (Hash != OutDesc.BaseMaterialHash && !OutDesc.ApprovedPermutationHashes.Contains(Hash))
+		UMaterialInstance const* MaterialInstance = MaterialInstances[Index];
+		if (MaterialInstance == nullptr)
 		{
-			UE_LOG(LogMaterialValidation, Log, TEXT("Add material instance '%s' to approved list."), *MaterialInstance->GetName());
-			OutDesc.ApprovedPermutationHashes.Add(Hash);
+			continue;
 		}
- 	}
-}
 
-/** Load the instance hierarchy for a material and find permutations that are not in the FMaterialValidationDesc. */
-static void GetNonValidMaterialInstances(UMaterial const& InMaterial, FMaterialValidationDesc const& InDesc, TArray<UMaterialInstance*>& OutMaterialInstances)
-{
-	TArray<UMaterialInstanceConstant*> MaterialInstances;
-	LoadMaterialInstanceConstants(InMaterial, MaterialInstances);
+		// Add the material instance to the map.
+		const FSoftObjectPath SoftPackagePath(MaterialInstance);
+		const FString PackagePath = SoftPackagePath.ToString();
 
-	FMaterialPermutationData PermutationData;
-	BuildPermutationData(InMaterial, PermutationData);
+		int32& AssetDataIndex = OutMaterialDesc.MaterialInstances.FindOrAdd(PackagePath);
+		AssetDataIndex = INDEX_NONE;
 
-	for (UMaterialInstanceConstant* MaterialInstance : MaterialInstances)
-	{
-		FMaterialValidationAssetData_MaterialInstance InstanceAssetData;
-		GetAssetData(*MaterialInstance, *MaterialInstance->Parent, PermutationData, InstanceAssetData);
+		// Get the material instance asset data.
+		FMaterialValidationAssetData_MaterialInstance MaterialInstanceAssetData;
+		GetAssetData(*MaterialInstance, PermutationData, MaterialInstanceAssetData);
+		const uint32 Hash = MaterialInstanceAssetData.PermutationHash;
 
-		const uint32 Hash = InstanceAssetData.PermutationHash;
-		if (Hash != InDesc.BaseMaterialHash && !InDesc.ApprovedPermutationHashes.Contains(Hash))
+		// Only store the asset data if the child differs from the parent.
+		// Otherwise we can assume the asset data is default and save some asset size by not storing it.
+		if (Hash != PermutationHashStack[CurrentDepth - 1])
 		{
-			OutMaterialInstances.Add(MaterialInstance);
+			AssetDataIndex = OutMaterialDesc.MaterialInstanceAssetDatas.Add(MaterialInstanceAssetData);
 		}
+
+		// If this is a new permutation hash then we store it and add to the shader count.
+		if (Hash != OutMaterialDesc.AssetData.PermutationHash && !OutMaterialDesc.PermutationHashes.Contains(Hash))
+		{
+			UE_LOG(LogMaterialValidation, Log, TEXT("Add material instance '%s' to permutaton list."), *PackagePath);
+			OutMaterialDesc.PermutationHashes.Add(Hash);
+			OutMaterialDesc.NumShadersTotal += GetNumShaderTypes(*MaterialInstance);
+		}
+
+		PermutationHashStack[CurrentDepth] = Hash;
 	}
 }
 
-/** Returns false if the material instance's parent material is in the validation group, but the material instance isn't in the approved list. */
-static bool IsValidationSuccess(UMaterialValidationGroup const& InGroup, UMaterialInstanceConstant const& InMaterialInstance)
+/** 
+ * Get the asset data associated with a package path.
+ * We get this from either the material validation group or from the object directly.
+ * We only set OutMaterialInstance if the object is loaded.
+ * For performance we avoid loading an object if it is unloaded (but could optionally do that if that becomes useful).
+ * If we return false then it is likely that the material instance exists but:
+ * - It wasn't stored in the validation group (it was probably created after last validation group update)
+ * - It isn't in the InObjectPackagePaths (this usually means that it isn't in the curent validation changelist)
+ * - It is not currently loaded.
+ */
+static bool GetMaterialInstanceAssetData(
+	FString const& InPackagePath,
+	TMap<FString, UMaterialInterface*> const& InObjectPackagePaths,
+	FMaterialValidationDesc const* InMaterialDesc,
+	FMaterialPermutationData const& InMaterialPermutationData,
+	FMaterialValidationAssetData_MaterialInstance& OutAssetData,
+	UMaterialInstanceConstant*& OutMaterialInstance)
 {
-	UMaterial const* Material = InMaterialInstance.GetMaterial();
-	const FSoftObjectPath SoftPackagePath(Material);
-	const FString PackagePath = SoftPackagePath.ToString();
+	// If we find in our known override object paths then initialize the asset data from the object.
+	if (UMaterialInterface* const* ModifiedAsset = InObjectPackagePaths.Find(InPackagePath))
+	{
+		if (UMaterialInstanceConstant* MaterialInstance = Cast<UMaterialInstanceConstant>(*ModifiedAsset))
+		{
+			GetAssetData(*MaterialInstance, InMaterialPermutationData, OutAssetData);
+			OutMaterialInstance = MaterialInstance;
+			return true;
+		}
+	}
 
-	FMaterialValidationDesc const* Settings = InGroup.Materials.Find(PackagePath);
-	if (Settings == nullptr)
+	// Otherwise use asset data from the material validation group.
+	if (InMaterialDesc != nullptr)
 	{
+		if (int32 const* AssetDataIndex = InMaterialDesc->MaterialInstances.Find(InPackagePath))
+		{
+			if (*AssetDataIndex == INDEX_NONE)
+			{
+				// Set up a default asset data with the right sizes.
+				const int32 StaticSwitchArrayNum = InMaterialDesc->AssetData.StaticSwitchValues.Num();
+				OutAssetData.StaticSwitchOverrideMask.SetNumZeroed(StaticSwitchArrayNum);
+				OutAssetData.StaticSwitchOverrideValues.SetNumZeroed(StaticSwitchArrayNum);
+				const int32 ComponentMaskArrayNum = InMaterialDesc->AssetData.ComponentMaskValues.Num();
+				OutAssetData.ComponentMaskOverrideMask.SetNumZeroed(ComponentMaskArrayNum);
+				OutAssetData.ComponentMaskOverrideValues.SetNumZeroed(ComponentMaskArrayNum);
+				return true;
+			}
+			else if (InMaterialDesc->MaterialInstanceAssetDatas.IsValidIndex(*AssetDataIndex))
+			{
+				OutAssetData = InMaterialDesc->MaterialInstanceAssetDatas[*AssetDataIndex];
+				return true;
+			}
+		}
+	}
+
+	// If we didn't find in the material validation group but the object is already loaded then use that.
+	FSoftObjectPath SoftObjectPath(InPackagePath);
+	TSoftObjectPtr<UMaterialInstanceConstant> SoftObjectPtr(SoftObjectPath);
+	if (UMaterialInstanceConstant* MaterialInstance = SoftObjectPtr.Get())
+	{
+		GetAssetData(*MaterialInstance, InMaterialPermutationData, OutAssetData);
+		OutMaterialInstance = MaterialInstance;
 		return true;
 	}
 
-	const uint32 Hash = BuildPermutationHash(InMaterialInstance, *Material);
-	return Hash == Settings->BaseMaterialHash || Settings->ApprovedPermutationHashes.Contains(Hash);
+	return false;
+}
+
+
+/**
+ * Create a transient material instance directly from an AssetData.
+ * This can be used to evaluate shader counts for an asset without loading it, and paying the cost of loading references etc.
+ * We do need the base material to be loaded, since we use it to create a temporarty material instance.
+ */
+static UMaterialInstanceConstant* CreateMaterialInstance(UMaterial& InBaseMaterial, FMaterialValidationAssetData_MaterialInstance const& InAssetData)
+{
+	// Create a temporary UMaterialInstanceConstant parented to our base material.
+	UMaterialInstanceConstant* MaterialInstance = NewObject<UMaterialInstanceConstant>();
+	MaterialInstance->SetParentEditorOnly(&InBaseMaterial);
+
+	FStaticParameterSet StaticParameterSet;
+
+	// Fill the static switch parameters.
+	TArray<FMaterialParameterInfo> StaticSwitchParameterInfo;
+	TArray<FGuid> StaticSwitchGuids;
+	InBaseMaterial.GetAllStaticSwitchParameterInfo(StaticSwitchParameterInfo, StaticSwitchGuids);
+
+	for (uint32 Index = 0; Index < (uint32)StaticSwitchParameterInfo.Num(); ++Index)
+	{
+		const bool bValue = (InAssetData.StaticSwitchOverrideValues[Index / 32] & (1 << Index)) != 0;
+		FMaterialParameterMetadata ParameterMetadata(bValue);
+		ParameterMetadata.ExpressionGuid = StaticSwitchGuids[Index];
+		StaticParameterSet.SetParameterValue(StaticSwitchParameterInfo[Index], ParameterMetadata, EMaterialSetParameterValueFlags::None);
+	}
+
+	// Fill the component mask parameters.
+	TArray<FMaterialParameterInfo> ComponentMaskParameterInfo;
+	TArray<FGuid> ComponentMaskGuids;
+	InBaseMaterial.GetAllStaticComponentMaskParameterInfo(ComponentMaskParameterInfo, ComponentMaskGuids);
+
+	for (uint32 Index = 0; Index < (uint32)ComponentMaskParameterInfo.Num(); ++Index)
+	{
+		const uint32 ValueMask = (InAssetData.ComponentMaskOverrideValues[Index / 8] >> ((Index % 8) * 4)) & 0xf;
+		const FStaticComponentMaskValue Value(ValueMask & 1, ValueMask & 2, ValueMask & 4, ValueMask & 8);
+		FMaterialParameterMetadata ParameterMetadata(Value);
+		ParameterMetadata.ExpressionGuid = ComponentMaskGuids[Index];
+		StaticParameterSet.SetParameterValue(ComponentMaskParameterInfo[Index], ParameterMetadata, EMaterialSetParameterValueFlags::None);
+	}
+
+	// Apply all of the base property overrides
+	FMaterialInstanceBasePropertyOverrides BaseProperties = GetStaticPropertiesAsBasePropertyOverrides(InAssetData.StaticProperties);
+
+	// Set all the static permutation state on the material instance.
+	//todo: We should handle material layer overrides too.
+	MaterialInstance->UpdateStaticPermutation(StaticParameterSet, BaseProperties, true);
+
+	return MaterialInstance;
 }
 
 } //namespace MaterialValidationLibrary
@@ -677,38 +884,7 @@ void UMaterialValidationLibrary::ResetGroup(UMaterialValidationGroup* InGroup)
 	}
 }
 
-void UMaterialValidationLibrary::AddMaterialsToGroup(UMaterialValidationGroup* InGroup, TArray<UMaterial*> const& InMaterials, bool bInUpdatePermutations)
-{
-	if (InGroup == nullptr)
-	{
-		return;
-	}
-	
-	UE_LOG(LogMaterialValidation, Display, TEXT("Adding %d materials."), InMaterials.Num());
-
-	FScopedSlowTask SlowTask(static_cast<float>(InMaterials.Num()), LOCTEXT("AddingMaterials", "Adding Materials..."), bInUpdatePermutations);
-
-	for (UMaterial* Material : InMaterials)
-	{
-		SlowTask.EnterProgressFrame();
-
-		if (Material != nullptr)
-		{
-			const FSoftObjectPath SoftPackagePath(Material);
-			const FString PackagePath = SoftPackagePath.ToString();
-
-			FMaterialValidationDesc& Desc = InGroup->Materials.FindOrAdd(PackagePath);
-			if (bInUpdatePermutations)
-			{
-				MaterialValidationLibrary::UpdateValidMaterialInstances(*Material, Desc);
-			}
-		}
-	}
-
-	InGroup->Modify();
-}
-
-void UMaterialValidationLibrary::AddMaterialsFromGroupDirectoryPaths(UMaterialValidationGroup* InGroup)
+void UMaterialValidationLibrary::AddMissingMaterialsToGroup(UMaterialValidationGroup* InGroup)
 {
 	if (InGroup == nullptr)
 	{
@@ -745,7 +921,7 @@ void UMaterialValidationLibrary::AddMaterialsFromGroupDirectoryPaths(UMaterialVa
 	}
 }
 
-void UMaterialValidationLibrary::RemoveMissingMaterialsFromGroup(UMaterialValidationGroup* InGroup)
+void UMaterialValidationLibrary::RemoveInvalidMaterialsFromGroup(UMaterialValidationGroup* InGroup)
 {
 	if (InGroup == nullptr)
 	{
@@ -801,7 +977,7 @@ void UMaterialValidationLibrary::RemoveMissingMaterialsFromGroup(UMaterialValida
 	}
 }
 
-void UMaterialValidationLibrary::UpdateAllMaterialsInGroup(UMaterialValidationGroup* InGroup)
+void UMaterialValidationLibrary::UpdateMaterialPermutationsInGroup(UMaterialValidationGroup* InGroup)
 {
 	if (InGroup == nullptr)
 	{
@@ -826,114 +1002,193 @@ void UMaterialValidationLibrary::UpdateAllMaterialsInGroup(UMaterialValidationGr
 	InGroup->Modify();
 }
 
-void UMaterialValidationLibrary::ValidateAllMaterialsInGroup(UMaterialValidationGroup* InGroup, TArray<UMaterialInstance*>& OutMaterialInstances)
+void UMaterialValidationLibrary::IsMaterialInGroup(UMaterialValidationGroup const* InGroup, UMaterial* InMaterial, bool& bOutIsInGroupPath, bool& bOutIsInGroup)
 {
-	if (InGroup == nullptr)
+	bOutIsInGroupPath = false;
+	bOutIsInGroup = false;
+
+	if (InGroup == nullptr || InMaterial == nullptr)
 	{
 		return;
 	}
 
-	UE_LOG(LogMaterialValidation, Display, TEXT("Validating %d materials."), InGroup->Materials.Num());
+	const FSoftObjectPath SoftPackagePath(InMaterial);
+	const FString PackagePath = SoftPackagePath.ToString();
 
-	FScopedSlowTask SlowTask(static_cast<float>(InGroup->Materials.Num()), LOCTEXT("ValidatingMaterials", "Validating Materials..."));
-	
-	for (TPair<FString, FMaterialValidationDesc>& It : InGroup->Materials)
+	if (InGroup->Materials.Contains(PackagePath))
 	{
-		SlowTask.EnterProgressFrame();
+		bOutIsInGroup = true;
+	}
 
-		TSoftObjectPtr<UMaterial> MaterialSoftPtr(MaterialValidationLibrary::ConvertStringPathToSoftObjectPath(It.Key));
-		if (UMaterial const* Material = MaterialSoftPtr.LoadSynchronous())
-		{
-			MaterialValidationLibrary::GetNonValidMaterialInstances(*Material, It.Value, OutMaterialInstances);
-		}
+	TArray<FString> SearchPaths;
+	TArray<FString> ExcludePaths;
+	MaterialValidationLibrary::GetSearchPaths(*InGroup, SearchPaths, ExcludePaths);
+
+	if (MaterialValidationLibrary::IsPackageInPaths(PackagePath, SearchPaths))
+	{
+		bOutIsInGroupPath = true;
+	}
+
+	if (bOutIsInGroupPath && MaterialValidationLibrary::IsPackageInPaths(PackagePath, ExcludePaths))
+	{
+		bOutIsInGroupPath = false;
 	}
 }
 
-void UMaterialValidationLibrary::AddMaterialInstanceToGroup(UMaterialValidationGroup* InGroup, UMaterialInstance* InMaterialInstance)
+void UMaterialValidationLibrary::IsMaterialInstanceInGroup(UMaterialValidationGroup const* InGroup, UMaterialInstanceConstant* InMaterialInstance, bool& bOutMaterialInGroup, bool& bOutMaterialPermutationInGroup)
 {
-	if (InGroup != nullptr)
+	bOutMaterialInGroup = false;
+	bOutMaterialPermutationInGroup = false;
+
+	if (InGroup == nullptr || InMaterialInstance == nullptr)
 	{
-		if (UMaterialInstanceConstant* MaterialInstance = Cast<UMaterialInstanceConstant>(InMaterialInstance))
-		{
-			UMaterial const* BaseMaterial = MaterialInstance->GetMaterial();
-			const FSoftObjectPath SoftPackagePath(BaseMaterial);
-			const FString PackagePath = SoftPackagePath.ToString();
+		return;
+	}
 
-			if (FMaterialValidationDesc* Settings = InGroup->Materials.Find(PackagePath))
-			{
-				const uint32 Hash = MaterialValidationLibrary::BuildPermutationHash(*MaterialInstance, *BaseMaterial);
-				if (Hash != Settings->BaseMaterialHash && !Settings->ApprovedPermutationHashes.Contains(Hash))
-				{
-					UE_LOG(LogMaterialValidation, Log, TEXT("Add material instance '%s' to approved list."), *MaterialInstance->GetName());
+	UMaterial const* Material = InMaterialInstance->GetMaterial();
+	const FSoftObjectPath SoftPackagePath(Material);
+	const FString PackagePath = SoftPackagePath.ToString();
 
-					Settings->ApprovedPermutationHashes.Add(Hash);
-					InGroup->Modify();
-				}
-			}
-		}
+	FMaterialValidationDesc const* MaterialDesc = InGroup->Materials.Find(PackagePath);
+	if (MaterialDesc == nullptr)
+	{
+		return;
 	}
+
+	bOutMaterialInGroup = true;
+
+	const uint32 Hash = MaterialValidationLibrary::BuildPermutationHash(*InMaterialInstance, *Material);
+	bOutMaterialPermutationInGroup = Hash == MaterialDesc->AssetData.PermutationHash || MaterialDesc->PermutationHashes.Contains(Hash);
 }
 
-bool UMaterialValidationLibrary::ValidateMaterialInstanceInGroup(UMaterialValidationGroup const* InGroup, UMaterialInstance* InMaterialInstance)
+int32 UMaterialValidationLibrary::GetShaderCount(UMaterialValidationGroup const* InGroup, UMaterial* InBaseMaterial)
 {
-	if (InGroup != nullptr)
+	if (InGroup == nullptr || InBaseMaterial == nullptr)
 	{
-		if (UMaterialInstanceConstant* MaterialInstance = Cast<UMaterialInstanceConstant>(InMaterialInstance))
-		{
-			return MaterialValidationLibrary::IsValidationSuccess(*InGroup, *MaterialInstance);
-		}
+		return 0;
 	}
-	
-	return true;
+
+	const FSoftObjectPath SoftPackagePath(InBaseMaterial);
+	const FString PackagePath = SoftPackagePath.ToString();
+
+	FMaterialValidationDesc const* MaterialDesc = InGroup->Materials.Find(PackagePath);
+	return MaterialDesc != nullptr ? MaterialDesc->NumShadersTotal : 0;
 }
 
-bool UMaterialValidationLibrary::ValidateMaterialInGroup(UMaterialValidationGroup const* InGroup, UMaterial* InMaterial)
+int32 UMaterialValidationLibrary::GetModifiedShaderCount(UMaterialValidationGroup const* InGroup, UMaterial* InBaseMaterial, TArray<UMaterialInterface*>& InModifiedObjects)
 {
-	if (InGroup == nullptr || InMaterial == nullptr)
+	if (InGroup == nullptr || InBaseMaterial == nullptr)
 	{
-		return true;
+		return 0;
 	}
 
-	const FSoftObjectPath SoftPackagePath(InMaterial);
-	const FString PackagePath = SoftPackagePath.ToString();
-
-	if (InGroup->Materials.Contains(PackagePath))
+	// Get paths to modified objects for lookup.
+	TMap<FString, UMaterialInterface*> ModifiedObjectPackagePaths;
+	for (UMaterialInterface* Object : InModifiedObjects)
 	{
-		return true;
+		const FSoftObjectPath SoftPackagePath(Object);
+		ModifiedObjectPackagePaths.Add(SoftPackagePath.ToString(), Object);
 	}
 
-	TArray<FString> SearchPaths;
-	TArray<FString> ExcludePaths;
-	MaterialValidationLibrary::GetSearchPaths(*InGroup, SearchPaths, ExcludePaths);
+	MaterialValidationLibrary::FMaterialPermutationData PermutationData;
+	BuildPermutationData(*InBaseMaterial, PermutationData);
 
-	if (MaterialValidationLibrary::IsPackageInPaths(PackagePath, ExcludePaths))
+	// Get the material asset data stored on the group.
+	// But evaluate from the material if the material is modified or if it is not found in the group.
+	const FSoftObjectPath  BaseMaterialSoftObjectPath(InBaseMaterial);
+	const FString BaseMaterialPackagePath = BaseMaterialSoftObjectPath.ToString();
+	FMaterialValidationDesc const* MaterialDesc = InGroup->Materials.Find(BaseMaterialPackagePath);
+
+	FMaterialValidationAssetData_Material MaterialAssetData;
+	if (MaterialDesc == nullptr || ModifiedObjectPackagePaths.Contains(BaseMaterialPackagePath))
 	{
-		return true;
+		GetAssetData(*InBaseMaterial, PermutationData, MaterialAssetData);
+	}
+	else
+	{
+		MaterialAssetData = MaterialDesc->AssetData;
 	}
 
-	if (MaterialValidationLibrary::IsPackageInPaths(PackagePath, SearchPaths))
+	// If the base material static parameters have changed from those stored on the group then any other material instance 
+	// asset data is no longer valid and we can't iterate the hierarchy without loading the material instances assets.
+	// So instead we approximate the change by only looking at the shader delta on the base material.
+	if (MaterialDesc != nullptr && MaterialDesc->AssetData.StaticPropertyLayoutHash != MaterialAssetData.StaticPropertyLayoutHash)
 	{
-		return false;
+		const int32 NumShadersBase = MaterialValidationLibrary::GetNumShaderTypes(*InBaseMaterial);
+		return MaterialDesc->NumShadersTotal + (NumShadersBase - MaterialDesc->NumShadersBase) * (1 + MaterialDesc->PermutationHashes.Num());
 	}
 
-	return true;
-}
+	// Get the material instance hierarchy from the asset registry.
+	TArray<FName> PackageNames;
+	TArray<int32> HierarchyDepths;
+	MaterialValidationLibrary::GetMaterialInstanceHierarchy(*InBaseMaterial, PackageNames, HierarchyDepths);
 
-int32 UMaterialValidationLibrary::GetApprovedPermutationCount(TSoftObjectPtr<UMaterial> InMaterial)
-{
-	TArray<UMaterialValidationGroup*> Groups;
-	GetAllGroups(Groups, true);
+	// Use a stack of asset data to walk the hierarchy with the asset data overrides applied at each level.
+	TArray<FMaterialValidationAssetData_MaterialInstance> AssetDataStack;
+
+	// Build the root asset data from the base material.
+	{
+		FMaterialValidationAssetData_MaterialInstance& Root = AssetDataStack.AddDefaulted_GetRef();
+		Root.StaticSwitchOverrideValues = MaterialAssetData.StaticSwitchValues;
+		Root.ComponentMaskOverrideValues = MaterialAssetData.ComponentMaskValues;
+		Root.StaticProperties = MaterialAssetData.StaticProperties;
+		Root.MaterialLayerHash = MaterialAssetData.MaterialLayerHash;
+		Root.PermutationHash = MaterialAssetData.PermutationHash;
+	}
+
+	// Initialize the shader count and unique hash array.
+	int32 NumShaders = MaterialValidationLibrary::GetNumShaderTypes(*InBaseMaterial);
+	TArray<uint32> PermutationHashes;
+	PermutationHashes.Add(MaterialAssetData.PermutationHash);
 
-	for (UMaterialValidationGroup* Group : Groups)
+	// Walk the material instance hierarchy.
+	for (int32 Index = 1; Index < PackageNames.Num(); ++Index)
 	{
-		if (FMaterialValidationDesc* Settings = Group->Materials.Find(InMaterial.ToSoftObjectPath().ToString()))
+		const int32 CurrentDepth = HierarchyDepths[Index];
+		FString PackagePath = MaterialValidationLibrary::AddPackagePathExtension(PackageNames[Index]);
+
+		FMaterialValidationAssetData_MaterialInstance AssetData;
+		UMaterialInstanceConstant* MaterialInstance = nullptr;
+		const bool bValidAssetData = GetMaterialInstanceAssetData(PackagePath, ModifiedObjectPackagePaths, MaterialDesc, PermutationData, AssetData, MaterialInstance);
+
+		if (!bValidAssetData)
 		{
-			return Settings->ApprovedPermutationHashes.Num();
+			// Ignore material instance and it's children.
+			while (Index < PackageNames.Num() - 1 && HierarchyDepths[Index + 1] > CurrentDepth)
+			{
+				++Index;
+			}
+			continue;
+		}
+
+		// Apply the current asset data overrides to the parent to get the final asset data for the child.
+		AssetDataStack.SetNum(CurrentDepth + 1);
+		FMaterialValidationAssetData_MaterialInstance& Child = AssetDataStack[CurrentDepth];
+		FMaterialValidationAssetData_MaterialInstance const& Parent = AssetDataStack[CurrentDepth - 1];
+
+		Child.StaticSwitchOverrideValues.SetNumZeroed(Parent.StaticSwitchOverrideValues.Num());
+		Child.ComponentMaskOverrideValues.SetNumZeroed(Parent.ComponentMaskOverrideValues.Num());
+		MaterialValidationLibrary::ApplyBitMaskOverrides(Parent.StaticSwitchOverrideValues, AssetData.StaticSwitchOverrideMask, AssetData.StaticSwitchOverrideValues, Child.StaticSwitchOverrideValues);
+		MaterialValidationLibrary::ApplyBitMaskOverrides(Parent.ComponentMaskOverrideValues, AssetData.ComponentMaskOverrideMask, AssetData.ComponentMaskOverrideValues, Child.ComponentMaskOverrideValues);
+		MaterialValidationLibrary::ApplyBasePropertyOverrides(Parent.StaticProperties, AssetData.StaticPropertyOverrideFlags, AssetData.StaticProperties, Child.StaticProperties);
+		Child.MaterialLayerHash = AssetData.MaterialLayerHash;
+
+		// Build the permutation hash and if this is a new permutation then add its shader count.
+		Child.PermutationHash = MaterialValidationLibrary::BuildPermutationHash(Child.StaticProperties, Child.StaticSwitchOverrideValues, Child.ComponentMaskOverrideValues, Child.MaterialLayerHash);
+		if (!PermutationHashes.Contains(Child.PermutationHash))
+		{
+			// If we don't have the material instance loaded then create a transient version.
+			if (MaterialInstance == nullptr)
+			{
+				MaterialInstance = MaterialValidationLibrary::CreateMaterialInstance(*InBaseMaterial, Child);
+			}
+
+			NumShaders += MaterialValidationLibrary::GetNumShaderTypes(*MaterialInstance);
+			PermutationHashes.Add(Child.PermutationHash);
 		}
 	}
 
-	// The Material was not found.
-	return 0;
+	return NumShaders;
 }
 
 #undef LOCTEXT_NAMESPACE
diff --git a/Engine/Plugins/Editor/MaterialValidation/Source/Private/MaterialValidationLibrary.h b/Engine/Plugins/Editor/MaterialValidation/Source/Private/MaterialValidationLibrary.h
--- a/Engine/Plugins/Editor/MaterialValidation/Source/Private/MaterialValidationLibrary.h
+++ b/Engine/Plugins/Editor/MaterialValidation/Source/Private/MaterialValidationLibrary.h
@@ -2,6 +2,7 @@
 
 #pragma once
 
+#include "Containers/ContainersFwd.h"
 #include "Kismet/BlueprintFunctionLibrary.h"
 #include "MaterialValidationLibrary.generated.h"
 
@@ -32,33 +33,23 @@ class UMaterialValidationLibrary : public UBlueprintFunctionLibrary
 	UFUNCTION(BlueprintCallable, Category = Material)
 	static void ResetGroup(UMaterialValidationGroup* InGroup);
 
-	/** 
-	 * Add UMaterials to a Material Validation Group.
-	 *
-	 * @param InGroup				The group to add materials to.
-	 * @param InMaterials			The array of materials to add.
-	 * @param bInUpdatePermutations	true to calculate and register all existing permutations for the materials (can be slow to load and analyze the tree of material instances).
-	 */
-	UFUNCTION(BlueprintCallable, Category = Material)
-	static void AddMaterialsToGroup(UMaterialValidationGroup* InGroup, TArray<UMaterial*> const& InMaterials, bool bInUpdatePermutations);
-
 	/**
 	 * Add all UMaterials to the Material Validation Group which are found in its material search directories.
 	 * This will only add materials, and not remove existing ones.
 	 * Materials are not updated with their permutations, so that newly found materials will have no approved permutations.
 	 *
-	 * @param InGroup				The group to add materials to.
+	 * @param InGroup	The group to add materials to.
 	 */
 	UFUNCTION(BlueprintCallable, Category = Material)
-	static void AddMaterialsFromGroupDirectoryPaths(UMaterialValidationGroup* InGroup);
+	static void AddMissingMaterialsToGroup(UMaterialValidationGroup* InGroup);
 
 	/**
 	 * Remove all UMaterials that no longer exist on disk from the Material Validation Group.
 	 *
-	 * @param InGroup	The group to update
+	 * @param InGroup	The group to update.
 	 */
 	UFUNCTION(BlueprintCallable, Category = Material)
-	static void RemoveMissingMaterialsFromGroup(UMaterialValidationGroup* InGroup);
+	static void RemoveInvalidMaterialsFromGroup(UMaterialValidationGroup* InGroup);
 
 	/**
 	 * Recalculate and reregister all existing permutations for all the existing materials in a Material Validation Group (can be slow to load and analyze the tree of material instances).
@@ -66,58 +57,49 @@ class UMaterialValidationLibrary : public UBlueprintFunctionLibrary
 	 * @param InGroup	The group to update.
 	 */
 	UFUNCTION(BlueprintCallable, Category = Material)
-	static void UpdateAllMaterialsInGroup(UMaterialValidationGroup* InGroup);
+	static void UpdateMaterialPermutationsInGroup(UMaterialValidationGroup* InGroup);
 
 	/**
-	 * Find Material Instance assets with permutations that aren't registered with their materials in a Material Validation Group (can be slow to load and analyze the tree of material instances).
-	 *
-	 * @param InGroup				The group to update.
-	 * @param OutMaterialInstances	Returned array of material instances that weren't marked as approved in the group.
-	 */
-	UFUNCTION(BlueprintCallable, Category = Material)
-	static void ValidateAllMaterialsInGroup(UMaterialValidationGroup* InGroup, TArray<UMaterialInstance*>& OutMaterialInstances);
-
-	/**
-	 * Add a material instance to a Material Validation Group.
-	 * It will only be added if it's base material is already in the group.
-	 *
-	 * @param InGroup				The group to update.
-	 * @param	InMaterialInstance	The material instance to add.
+	 * Get whether a material belongs to a group.
+	 * 
+	 * @param	InGroup				The group to use.
+	 * @param	InMaterial			The material to validate.
+	 * @param	bOutIsInGroupPath	Returns true if the material is in the validation group search path.
+	 * @param	bOutIsInGroup		Returns true if the material is in the validation group material list.
 	 */
 	UFUNCTION(BlueprintCallable, Category = Material)
-	static void AddMaterialInstanceToGroup(UMaterialValidationGroup* InGroup, UMaterialInstance* InMaterialInstance);
+	static void IsMaterialInGroup(UMaterialValidationGroup const* InGroup, UMaterial* InMaterial, bool& bOutIsInGroupPath, bool& bOutIsInGroup);
 
 	/**
-	 * Validate whether a material instance is approved.
-	 * 
-	 * @param	InGroup				The group to use.
-	 * @param	InMaterialInstance	The material instance to validate.
+	 * Get whether a material instance belongs to the group and whether it has a known permutation.
 	 * 
-	 * @returns	false if the material instance belongs to the validation group, but is not in the approved list.
+	 * @param	InGroup							The group to use.
+	 * @param	InMaterialInstance				The material instance to validate.
+	 * @param	bOutMaterialInGroup				Returns true if the base material is in the validation group material list.
+	 * @param	bOutMaterialPermutationInGroup	Returns true if the material instance permutation is in the permutation list.
 	 */
 	UFUNCTION(BlueprintCallable, Category = Material)
-	static bool ValidateMaterialInstanceInGroup(UMaterialValidationGroup const* InGroup, UMaterialInstance* InMaterialInstance);
+	static void IsMaterialInstanceInGroup(UMaterialValidationGroup const* InGroup, UMaterialInstanceConstant* InMaterialInstance, bool& bOutMaterialInGroup, bool& bOutMaterialPermutationInGroup);
 
 	/**
-	 * Validate whether a material is approved.
-	 * 
-	 * @param	InGroup		The group to use.
-	 * @param	InMaterial	The material to validate.
-	 * 
-	 * @returns	false if the material is in the validation group search path, but is not in the material list.
+	 * Get the number of shaders used by all permutations of a base material using only the information stored in the group.
+	 *
+	 * @param	InGroup					The group to use.
+	 * @param	InMaterial				The base material to evaluate.
+	 *
+	 * @returns	The total number of shaders calculated at the last time we updated the group.
 	 */
-	UFUNCTION(BlueprintCallable, Category = Material)
-	static bool ValidateMaterialInGroup(UMaterialValidationGroup const* InGroup, UMaterial* InMaterial);
+	static int32 GetShaderCount(UMaterialValidationGroup const* InGroup, UMaterial* InMaterial);
 
 	/**
-	 * Get the number of approved permutations for a given UMaterial.
-	 * This gets the result directly from the available Material Validation Group data and
-	 * doesn't load or interrogate any material instances.
-	 * 
-	 * @param	InMaterial	The material to get the result for.
+	 * Get the number of shaders used by all permutations of a base material taking into account local changes in the materials.
+	 *
+	 * @param	InGroup				The group to use.
+	 * @param	InMaterial			The base material to evaluate.
+	 * @param	InModifiedObjects	The array of material and material instance objects that need evaluating with their local changes.
 	 * 
-	 * @returns	The number of approved permutations found. This may be less than the number of actual permutations.
+	 * @returns	The total number of shaders. This value is the best estimate that we can give without triggering asset loads.
 	 */
 	UFUNCTION(BlueprintCallable, Category = Material)
-	static int32 GetApprovedPermutationCount(TSoftObjectPtr<UMaterial> InMaterial);
+	static int32 GetModifiedShaderCount(UMaterialValidationGroup const* InGroup, UMaterial* InMaterial, TArray<UMaterialInterface*>& InModifiedObjects);
 };
diff --git a/Engine/Plugins/Editor/MaterialValidation/Source/Private/MaterialValidators.cpp b/Engine/Plugins/Editor/MaterialValidation/Source/Private/MaterialValidators.cpp
--- a/Engine/Plugins/Editor/MaterialValidation/Source/Private/MaterialValidators.cpp
+++ b/Engine/Plugins/Editor/MaterialValidation/Source/Private/MaterialValidators.cpp
@@ -2,11 +2,11 @@
 
 #include "MaterialValidators.h"
 
-#include "AssetRegistry/AssetDataToken.h"
 #include "Editor.h"
 #include "Logging/TokenizedMessage.h"
-#include "MaterialShared.h"
+#include "MaterialEditingLibrary.h"
 #include "Materials/Material.h"
+#include "Materials/MaterialFunction.h"
 #include "Materials/MaterialInstanceConstant.h"
 #include "MaterialValidationConfig.h"
 #include "MaterialValidationGroup.h"
@@ -17,30 +17,6 @@
 
 #define LOCTEXT_NAMESPACE "MaterialValidation"
 
-/**
- * Get the number of shaders associated with a material instance.
- * These _may_ be deduplicated with other material instances, so this is an upper
- * limit estimate of the cost.
- * todo: Calculate for other platforms, or at least the current preview platform.
- * todo: Can we get an estimate of memory cost too?
- */
-static int32 GetShaderCount(UMaterialInterface* InMaterial)
-{
-	TArray<FDebugShaderTypeInfo> ShaderInfos;
-	InMaterial->GetShaderTypes(GMaxRHIShaderPlatform, nullptr, ShaderInfos);
-
-	int NumShaders = 0;
-	for (FDebugShaderTypeInfo const& ShaderInfo : ShaderInfos)
-	{
-		NumShaders += ShaderInfo.ShaderTypes.Num();
-		for (FDebugShaderPipelineInfo const& PipelineInfo : ShaderInfo.Pipelines)
-		{
-			NumShaders += PipelineInfo.ShaderTypes.Num();
-		}
-	}
-	return NumShaders;
-}
-
 bool UEditorValidator_MaterialPermutation::CanValidateAsset_Implementation(const FAssetData& InAssetData, UObject* InAsset, FDataValidationContext& InContext) const
 {
 	const EDataValidationUsecase Usecase = InContext.GetValidationUsecase();
@@ -48,7 +24,7 @@ bool UEditorValidator_MaterialPermutation::CanValidateAsset_Implementation(const
 	{
 		return false;
 	}
-	if (Cast<UMaterialInstanceConstant>(InAsset) == nullptr && Cast<UMaterial>(InAsset) == nullptr)
+	if (Cast<UMaterialInstanceConstant>(InAsset) == nullptr && Cast<UMaterial>(InAsset) == nullptr && Cast<UMaterialFunction>(InAsset) == nullptr)
 	{
 		return false;
 	}
@@ -58,52 +34,68 @@ bool UEditorValidator_MaterialPermutation::CanValidateAsset_Implementation(const
 EDataValidationResult UEditorValidator_MaterialPermutation::ValidateLoadedAsset_Implementation(const FAssetData& InAssetData, UObject* InAsset, FDataValidationContext& InContext)
 {
 	UMaterial* Material = Cast<UMaterial>(InAsset);
-	UMaterialInstance* MaterialInstance = Cast<UMaterialInstance>(InAsset);
+	UMaterialInstanceConstant* MaterialInstance = Cast<UMaterialInstanceConstant>(InAsset);
+	UMaterialFunction* MaterialFunction = Cast<UMaterialFunction>(InAsset);
 
 	if (Material || MaterialInstance)
 	{
+		int32 NumShaders = 0;
+
 		TArray<UMaterialValidationGroup*> Groups;
 		UMaterialValidationLibrary::GetAllGroups(Groups, /*bInSyncLoad*/true);
 
-		int32 NumShaders = 0;
-
 		for (UMaterialValidationGroup* Group : Groups)
 		{
-			if (Material && !UMaterialValidationLibrary::ValidateMaterialInGroup(Group, Material))
+			if (Material)
 			{
-				NumShaders = GetShaderCount(Material);
-				break;
+				bool bIsInGroupPath = false;
+				bool bIsInGroup = false;
+				UMaterialValidationLibrary::IsMaterialInGroup(Group, Material, bIsInGroupPath, bIsInGroup);
+				
+				if (bIsInGroupPath)
+				{
+					BaseMaterials.AddUnique(Material);
+					ModifiedObjects.Add(Material);
+
+					if (!bIsInGroup)
+					{
+						NumShaders = UMaterialEditingLibrary::GetNumShaderTypes(Material);
+					}
+					break;
+				}
 			}
 			
-			if (MaterialInstance && !UMaterialValidationLibrary::ValidateMaterialInstanceInGroup(Group, MaterialInstance))
+			if (MaterialInstance)
 			{
-				NumShaders = GetShaderCount(MaterialInstance);
-				break;
+				bool bMaterialInGroup = false;
+				bool bMaterialPermutationInGroup = false;
+				UMaterialValidationLibrary::IsMaterialInstanceInGroup(Group, MaterialInstance, bMaterialInGroup, bMaterialPermutationInGroup);
+
+				if (bMaterialInGroup) 
+				{
+					BaseMaterials.AddUnique(MaterialInstance->GetMaterial());
+					ModifiedObjects.Add(MaterialInstance);
+
+					if (!bMaterialPermutationInGroup)
+					{
+						NumShaders = UMaterialEditingLibrary::GetNumShaderTypes(MaterialInstance);
+					}
+					break;
+				}
 			}
 		}
 
 		if (NumShaders > 0)
 		{
-			AccumulatedShaderCount += NumShaders;
-
-			// Need to register callback to reset AccumulatedShaderCount after validation is complete.
-			if (!OnPostValidationHandle.IsValid())
-			{
-				OnPostValidationHandle = FEditorDelegates::OnPostAssetValidation.AddUObject(this, &UEditorValidator_MaterialPermutation::OnPostValidation);
-			}
-
-			AssetMessage(EMessageSeverity::Info)->AddText(LOCTEXT("ShaderBudget_Info", "Adds {0} new shaders"), FText::AsNumber(NumShaders));
-			AssetMessage(EMessageSeverity::Info)->AddText(FText::Format(INVTEXT("ShaderBudgetDelta: {0}"), AccumulatedShaderCount));
-
-			UMaterialValidationConfig const* Config = GetDefault<UMaterialValidationConfig>();
-			const int32 Threshold = Config->ShaderBudgetThreshold;
-			if (Threshold > 0 && AccumulatedShaderCount >= Threshold)
-			{
-				AssetMessage(EMessageSeverity::Info)->AddText(INVTEXT("ShaderBudgetIntegrationRequest"));
-			}
+			AssetMessage(EMessageSeverity::Info)->AddText(FText::Format(LOCTEXT("ShaderBudget_Info", "Adds {0} new shaders"), FText::AsNumber(NumShaders)));
 		}
 	}
 
+	if (MaterialFunction)
+	{
+		//todo: Add UMaterial refs to BaseMaterials. Note this will require loading assets, but is probably OK.
+	}
+
 	if (GetValidationResult() != EDataValidationResult::Invalid)
 	{
 		AssetPasses(InAsset);
@@ -112,28 +104,66 @@ EDataValidationResult UEditorValidator_MaterialPermutation::ValidateLoadedAsset_
 	return GetValidationResult();
 }
 
-void UEditorValidator_MaterialPermutation::OnPostValidation()
+void UEditorValidator_MaterialPermutation::OnPostAssetValidation(TArray<TSharedRef<FTokenizedMessage>>& OutMessages)
 {
-	AccumulatedShaderCount = 0;
+	// Now that we have all the touched materials for the changelist we can get the shader delta per base material hierarchy.
+	// We sum these to get the final shader delta to report for validation.
+	int32 NumShaders = 0;
+
+	TArray<UMaterialValidationGroup*> Groups;
+	UMaterialValidationLibrary::GetAllGroups(Groups, /*bInSyncLoad*/true);
+
+	for (UMaterial* BaseMaterial : BaseMaterials)
+	{
+		for (UMaterialValidationGroup* Group : Groups)
+		{
+			bool bIsInGroupPath = false;
+			bool bIsInGroup = false;
+			UMaterialValidationLibrary::IsMaterialInGroup(Group, BaseMaterial, bIsInGroupPath, bIsInGroup);
 
-	// Remove delegate until next validation.
-	FEditorDelegates::OnPostAssetValidation.Remove(OnPostValidationHandle);
-	OnPostValidationHandle.Reset();
+			if (bIsInGroupPath)
+			{
+				NumShaders -= UMaterialValidationLibrary::GetShaderCount(Group, BaseMaterial);
+				NumShaders += UMaterialValidationLibrary::GetModifiedShaderCount(Group, BaseMaterial, MutableView(ModifiedObjects));
+			}
+		}
+	}
+
+	if (NumShaders != 0)
+	{
+		FNumberFormattingOptions NumberFormat;
+		NumberFormat.UseGrouping = false;
+		OutMessages.Add(FTokenizedMessage::Create(EMessageSeverity::Info, FText::Format(INVTEXT("ShaderBudgetDelta: {0}"), FText::AsNumber(NumShaders, &NumberFormat))));
+
+		UMaterialValidationConfig const* Config = GetDefault<UMaterialValidationConfig>();
+		const int32 Threshold = Config->ShaderBudgetThreshold;
+		if (Threshold > 0 && NumShaders >= Threshold)
+		{
+			OutMessages.Add(FTokenizedMessage::Create(EMessageSeverity::Info, INVTEXT("ShaderBudgetIntegrationRequest")));
+		}
+	}
+
+	BaseMaterials.Reset();
+	ModifiedObjects.Reset();
 }
 
 bool UMaterialEditorValidator_Permutation::Validate(UMaterialInterface* InMaterial, FMaterialEditorValidatorContext& InContext) const
 {
-	if (UMaterialInstance* MaterialInstance = Cast<UMaterialInstance>(InMaterial))
+	if (UMaterialInstanceConstant* MaterialInstance = Cast<UMaterialInstanceConstant>(InMaterial))
 	{
 		TArray<UMaterialValidationGroup*> Groups;
 		UMaterialValidationLibrary::GetAllGroups(Groups, /*bInSyncLoad*/false);
 
 		for (UMaterialValidationGroup const* Group : Groups)
 		{
-			if (!UMaterialValidationLibrary::ValidateMaterialInstanceInGroup(Group, MaterialInstance))
+			bool bMaterialInGroup = false;
+			bool bMaterialPermutationInGroup = false;
+			UMaterialValidationLibrary::IsMaterialInstanceInGroup(Group, MaterialInstance, bMaterialInGroup, bMaterialPermutationInGroup);
+
+			if (bMaterialInGroup && !bMaterialPermutationInGroup)
 			{
 				UMaterial const* BaseMaterial = MaterialInstance->GetMaterial();
-				const int32 NumShaders = GetShaderCount(MaterialInstance);
+				const int32 NumShaders = UMaterialEditingLibrary::GetNumShaderTypes(MaterialInstance);
 
 				InContext.Message = FText::Format(
 					LOCTEXT("MaterialEditorValidation_Warn", "This material instance creates a new permutation of Base Material '{0}'.\nA permutation of this material is likely to generate {1} unique shaders.\nCan you use an existing permutation instead?"),
diff --git a/Engine/Plugins/Editor/MaterialValidation/Source/Private/MaterialValidators.h b/Engine/Plugins/Editor/MaterialValidation/Source/Private/MaterialValidators.h
--- a/Engine/Plugins/Editor/MaterialValidation/Source/Private/MaterialValidators.h
+++ b/Engine/Plugins/Editor/MaterialValidation/Source/Private/MaterialValidators.h
@@ -15,12 +15,15 @@ class UEditorValidator_MaterialPermutation : public UEditorValidatorBase
 
 	virtual bool CanValidateAsset_Implementation(const FAssetData& InAssetData, UObject* InAsset, FDataValidationContext& InContext) const override;
 	virtual EDataValidationResult ValidateLoadedAsset_Implementation(const FAssetData& InAssetData, UObject* InAsset, FDataValidationContext& InContext) override;
+	virtual void OnPostAssetValidation(TArray<TSharedRef<FTokenizedMessage>>& OutMessages) override;
 
 private:
-	int32 AccumulatedShaderCount = 0;
-
-	FDelegateHandle OnPostValidationHandle;
-	void OnPostValidation();
+	/** Collected base materials for which to validate the full hierarchy after individual asset validation is complete. */
+	UPROPERTY()
+	TArray<TObjectPtr<class UMaterial>> BaseMaterials;
+	/** Collected individual assets whose changes should be taken into account when doing full hierarchy validation. */
+	UPROPERTY()
+	TArray<TObjectPtr<class UMaterialInterface>> ModifiedObjects;
 };
 
 /** Implementation of UMaterialEditorValidatorBase that applies validation for limiting the addition of material permutations. */
