Fix CIS failure from duplicated function name.
#rnx
#Synced-CL 48144554

diff --git a/Engine/Plugins/Editor/DataValidation/Source/DataValidation/Private/EditorValidatorSubsystem.cpp b/Engine/Plugins/Editor/DataValidation/Source/DataValidation/Private/EditorValidatorSubsystem.cpp
--- a/Engine/Plugins/Editor/DataValidation/Source/DataValidation/Private/EditorValidatorSubsystem.cpp
+++ b/Engine/Plugins/Editor/DataValidation/Source/DataValidation/Private/EditorValidatorSubsystem.cpp
@@ -832,7 +832,7 @@ EDataValidationResult UEditorValidatorSubsystem::ValidateAssetsInternal(
 
 	ForEachEnabledValidator([&](UEditorValidatorBase* Validator)
 	{
-		Validator->OnPostAssetValidation(OutResults.ValidatorMessages);
+		Validator->PostAssetValidation(OutResults.ValidatorMessages);
 		return true;
 	});
 
diff --git a/Engine/Plugins/Editor/DataValidation/Source/DataValidation/Public/EditorValidatorBase.h b/Engine/Plugins/Editor/DataValidation/Source/DataValidation/Public/EditorValidatorBase.h
--- a/Engine/Plugins/Editor/DataValidation/Source/DataValidation/Public/EditorValidatorBase.h
+++ b/Engine/Plugins/Editor/DataValidation/Source/DataValidation/Public/EditorValidatorBase.h
@@ -58,7 +58,7 @@ public:
 	}
 	
 	/** Override this to do any state cleanup or add additional validation messages after all per asset validation is complete. */
-	virtual void OnPostAssetValidation(TArray<TSharedRef<FTokenizedMessage>>& OutMessages) {}
+	virtual void PostAssetValidation(TArray<TSharedRef<FTokenizedMessage>>& OutMessages) {}
 
 	/** Override this to determine whether or not you can use this validator given this usecase */
 	UFUNCTION(BlueprintNativeEvent, Category = "Asset Validation", DisplayName="Can Validate")
diff --git a/Engine/Plugins/Editor/MaterialValidation/Source/Private/MaterialValidators.cpp b/Engine/Plugins/Editor/MaterialValidation/Source/Private/MaterialValidators.cpp
--- a/Engine/Plugins/Editor/MaterialValidation/Source/Private/MaterialValidators.cpp
+++ b/Engine/Plugins/Editor/MaterialValidation/Source/Private/MaterialValidators.cpp
@@ -104,7 +104,7 @@ EDataValidationResult UEditorValidator_MaterialPermutation::ValidateLoadedAsset_
 	return GetValidationResult();
 }
 
-void UEditorValidator_MaterialPermutation::OnPostAssetValidation(TArray<TSharedRef<FTokenizedMessage>>& OutMessages)
+void UEditorValidator_MaterialPermutation::PostAssetValidation(TArray<TSharedRef<FTokenizedMessage>>& OutMessages)
 {
 	// Now that we have all the touched materials for the changelist we can get the shader delta per base material hierarchy.
 	// We sum these to get the final shader delta to report for validation.
diff --git a/Engine/Plugins/Editor/MaterialValidation/Source/Private/MaterialValidators.h b/Engine/Plugins/Editor/MaterialValidation/Source/Private/MaterialValidators.h
--- a/Engine/Plugins/Editor/MaterialValidation/Source/Private/MaterialValidators.h
+++ b/Engine/Plugins/Editor/MaterialValidation/Source/Private/MaterialValidators.h
@@ -15,7 +15,7 @@ class UEditorValidator_MaterialPermutation : public UEditorValidatorBase
 
 	virtual bool CanValidateAsset_Implementation(const FAssetData& InAssetData, UObject* InAsset, FDataValidationContext& InContext) const override;
 	virtual EDataValidationResult ValidateLoadedAsset_Implementation(const FAssetData& InAssetData, UObject* InAsset, FDataValidationContext& InContext) override;
-	virtual void OnPostAssetValidation(TArray<TSharedRef<FTokenizedMessage>>& OutMessages) override;
+	virtual void PostAssetValidation(TArray<TSharedRef<FTokenizedMessage>>& OutMessages) override;
 
 private:
 	/** Collected base materials for which to validate the full hierarchy after individual asset validation is complete. */
