Material Validation. Add seperate show and approval thresholds.
Disable post validation work if neither threshold is set.
#rnx
#Synced-CL 48145532

diff --git a/Engine/Plugins/Editor/MaterialValidation/Source/Private/MaterialValidationConfig.h b/Engine/Plugins/Editor/MaterialValidation/Source/Private/MaterialValidationConfig.h
--- a/Engine/Plugins/Editor/MaterialValidation/Source/Private/MaterialValidationConfig.h
+++ b/Engine/Plugins/Editor/MaterialValidation/Source/Private/MaterialValidationConfig.h
@@ -23,9 +23,13 @@ public:
 	UPROPERTY(config, EditAnywhere, Category = MaterialPermutationBudget, meta = (DisplayName = "Enable Validation"))
 	bool bEnable = false;
 
+	/** Threshold for shader count increase before showing to user. */
+	UPROPERTY(config, EditAnywhere, Category = MaterialPermutationBudget, meta = (EditCondition = "bEnable"))
+	int32 ShaderBudgetShowThreshold = 0;
+
 	/** Threshold for shader count increase before triggering further submit approval. */
 	UPROPERTY(config, EditAnywhere, Category = MaterialPermutationBudget, meta = (EditCondition = "bEnable"))
-	int32 ShaderBudgetThreshold = 0;
+	int32 ShaderBudgetApprovalThreshold = 0;
 
 	/** An array of Material Validation Groups to use during validation. */
 	UPROPERTY(config, EditAnywhere, Category = MaterialPermutationBudget, meta = (EditCondition = "bEnable", DisplayName = "Material Validation Groups"))
diff --git a/Engine/Plugins/Editor/MaterialValidation/Source/Private/MaterialValidators.cpp b/Engine/Plugins/Editor/MaterialValidation/Source/Private/MaterialValidators.cpp
--- a/Engine/Plugins/Editor/MaterialValidation/Source/Private/MaterialValidators.cpp
+++ b/Engine/Plugins/Editor/MaterialValidation/Source/Private/MaterialValidators.cpp
@@ -110,37 +110,43 @@ void UEditorValidator_MaterialPermutation::PostAssetValidation(TArray<TSharedRef
 	// We sum these to get the final shader delta to report for validation.
 	int32 NumShaders = 0;
 
-	TArray<UMaterialValidationGroup*> Groups;
-	UMaterialValidationLibrary::GetAllGroups(Groups, /*bInSyncLoad*/true);
+	UMaterialValidationConfig const* Config = GetDefault<UMaterialValidationConfig>();
+	const int32 ShowThreshold = Config->ShaderBudgetShowThreshold;
+	const int32 ApprovalThreshold = Config->ShaderBudgetApprovalThreshold;
 
-	for (UMaterial* BaseMaterial : BaseMaterials)
+	// Only do the work of calculating the shader delta if we have one of the thresholds set.
+	if (ShowThreshold > 0 || ApprovalThreshold > 0)
 	{
-		for (UMaterialValidationGroup* Group : Groups)
-		{
-			bool bIsInGroupPath = false;
-			bool bIsInGroup = false;
-			UMaterialValidationLibrary::IsMaterialInGroup(Group, BaseMaterial, bIsInGroupPath, bIsInGroup);
+		TArray<UMaterialValidationGroup*> Groups;
+		UMaterialValidationLibrary::GetAllGroups(Groups, /*bInSyncLoad*/true);
 
-			if (bIsInGroupPath)
+		for (UMaterial* BaseMaterial : BaseMaterials)
+		{
+			for (UMaterialValidationGroup* Group : Groups)
 			{
-				NumShaders -= UMaterialValidationLibrary::GetShaderCount(Group, BaseMaterial);
-				NumShaders += UMaterialValidationLibrary::GetModifiedShaderCount(Group, BaseMaterial, MutableView(ModifiedObjects));
+				bool bIsInGroupPath = false;
+				bool bIsInGroup = false;
+				UMaterialValidationLibrary::IsMaterialInGroup(Group, BaseMaterial, bIsInGroupPath, bIsInGroup);
+
+				if (bIsInGroupPath)
+				{
+					NumShaders -= UMaterialValidationLibrary::GetShaderCount(Group, BaseMaterial);
+					NumShaders += UMaterialValidationLibrary::GetModifiedShaderCount(Group, BaseMaterial, MutableView(ModifiedObjects));
+				}
 			}
 		}
 	}
 
-	if (NumShaders != 0)
+	if (ShowThreshold > 0 && FMath::Abs(NumShaders) >= ShowThreshold)
 	{
 		FNumberFormattingOptions NumberFormat;
 		NumberFormat.UseGrouping = false;
 		OutMessages.Add(FTokenizedMessage::Create(EMessageSeverity::Info, FText::Format(INVTEXT("ShaderBudgetDelta: {0}"), FText::AsNumber(NumShaders, &NumberFormat))));
+	}
 
-		UMaterialValidationConfig const* Config = GetDefault<UMaterialValidationConfig>();
-		const int32 Threshold = Config->ShaderBudgetThreshold;
-		if (Threshold > 0 && NumShaders >= Threshold)
-		{
-			OutMessages.Add(FTokenizedMessage::Create(EMessageSeverity::Info, INVTEXT("ShaderBudgetIntegrationRequest")));
-		}
+	if (ApprovalThreshold > 0 && NumShaders >= ApprovalThreshold)
+	{
+		OutMessages.Add(FTokenizedMessage::Create(EMessageSeverity::Info, INVTEXT("ShaderBudgetIntegrationRequest")));
 	}
 
 	BaseMaterials.Reset();
