Material Validation. Use cheaper algorithm to calculate shader count for non-loaded material instances.
Calling UpdateStaticPermutation() is slow when we call it 100s of times for a big hierarchy and can add 1-2 minutes of validation time.
Instead rely on the fact that static switches don't impact shader count (for now) and only apply only material instance base properties which is cheap.
#rnx
#Synced-CL 48160268

diff --git a/Engine/Plugins/Editor/MaterialValidation/Source/Private/MaterialValidationLibrary.cpp b/Engine/Plugins/Editor/MaterialValidation/Source/Private/MaterialValidationLibrary.cpp
--- a/Engine/Plugins/Editor/MaterialValidation/Source/Private/MaterialValidationLibrary.cpp
+++ b/Engine/Plugins/Editor/MaterialValidation/Source/Private/MaterialValidationLibrary.cpp
@@ -794,6 +794,13 @@ static UMaterialInstanceConstant* CreateMaterialInstance(UMaterial& InBaseMateri
 	UMaterialInstanceConstant* MaterialInstance = NewObject<UMaterialInstanceConstant>();
 	MaterialInstance->SetParentEditorOnly(&InBaseMaterial);
 
+	// Apply all of the static properties.
+	FMaterialInstanceBasePropertyOverrides BaseProperties = GetStaticPropertiesAsBasePropertyOverrides(InAssetData.StaticProperties);
+
+	// Calling UpdateStaticPermutation() is slow when we call it 100s of times for a big hierarchy.
+	// Maybe the static switches can't affect the number of shaders? Even if that is true it might not always be true in future...
+	// For now we can take the faster branch below which only sets the base properties (which can modify the shader count).
+#if 0 
 	FStaticParameterSet StaticParameterSet;
 
 	// Fill the static switch parameters.
@@ -823,12 +830,16 @@ static UMaterialInstanceConstant* CreateMaterialInstance(UMaterial& InBaseMateri
 		StaticParameterSet.SetParameterValue(ComponentMaskParameterInfo[Index], ParameterMetadata, EMaterialSetParameterValueFlags::None);
 	}
 
-	// Apply all of the base property overrides
-	FMaterialInstanceBasePropertyOverrides BaseProperties = GetStaticPropertiesAsBasePropertyOverrides(InAssetData.StaticProperties);
+	//todo: We should handle material layer overrides too.
 
 	// Set all the static permutation state on the material instance.
-	//todo: We should handle material layer overrides too.
 	MaterialInstance->UpdateStaticPermutation(StaticParameterSet, BaseProperties, true);
+#else
+	// Fast path only sets base properties which is still picked up by GetNumShaderTypes()
+	MaterialInstance->BasePropertyOverrides = BaseProperties;
+	MaterialInstance->UpdateOverridableBaseProperties();
+	MaterialInstance->bHasStaticPermutationResource = true;
+#endif
 
 	return MaterialInstance;
 }
