Fix material validation crash for case where base material had static switches added but isn't part of the changelist.
#rnx
#Synced-CL 48236586

diff --git a/Engine/Plugins/Editor/MaterialValidation/Source/Private/MaterialValidationGroup.h b/Engine/Plugins/Editor/MaterialValidation/Source/Private/MaterialValidationGroup.h
--- a/Engine/Plugins/Editor/MaterialValidation/Source/Private/MaterialValidationGroup.h
+++ b/Engine/Plugins/Editor/MaterialValidation/Source/Private/MaterialValidationGroup.h
@@ -6,8 +6,6 @@
 #include "MaterialValidationAssetData.h"
 #include "MaterialValidationGroup.generated.h"
 
-class UMaterial;
-
 /** Description of validation setup for a single UMaterial. */
 USTRUCT(MinimalAPI)
 struct FMaterialValidationDesc
diff --git a/Engine/Plugins/Editor/MaterialValidation/Source/Private/MaterialValidationLibrary.cpp b/Engine/Plugins/Editor/MaterialValidation/Source/Private/MaterialValidationLibrary.cpp
--- a/Engine/Plugins/Editor/MaterialValidation/Source/Private/MaterialValidationLibrary.cpp
+++ b/Engine/Plugins/Editor/MaterialValidation/Source/Private/MaterialValidationLibrary.cpp
@@ -782,11 +782,10 @@ static bool GetMaterialInstanceAssetData(
 	return false;
 }
 
-
 /**
  * Create a transient material instance directly from an AssetData.
  * This can be used to evaluate shader counts for an asset without loading it, and paying the cost of loading references etc.
- * We do need the base material to be loaded, since we use it to create a temporarty material instance.
+ * We do need the base material to be loaded, since we use it to create a temporary material instance.
  */
 static UMaterialInstanceConstant* CreateMaterialInstance(UMaterial& InBaseMaterial, FMaterialValidationAssetData_MaterialInstance const& InAssetData)
 {
@@ -1104,6 +1103,9 @@ int32 UMaterialValidationLibrary::GetModifiedShaderCount(UMaterialValidationGrou
 	MaterialValidationLibrary::FMaterialPermutationData PermutationData;
 	BuildPermutationData(*InBaseMaterial, PermutationData);
 
+	FMaterialValidationAssetData_Material MaterialAssetDataFromObject;
+	GetAssetData(*InBaseMaterial, PermutationData, MaterialAssetDataFromObject);
+
 	// Get the material asset data stored on the group.
 	// But evaluate from the material if the material is modified or if it is not found in the group.
 	const FSoftObjectPath  BaseMaterialSoftObjectPath(InBaseMaterial);
@@ -1113,7 +1115,7 @@ int32 UMaterialValidationLibrary::GetModifiedShaderCount(UMaterialValidationGrou
 	FMaterialValidationAssetData_Material MaterialAssetData;
 	if (MaterialDesc == nullptr || ModifiedObjectPackagePaths.Contains(BaseMaterialPackagePath))
 	{
-		GetAssetData(*InBaseMaterial, PermutationData, MaterialAssetData);
+		MaterialAssetData = MaterialAssetDataFromObject;
 	}
 	else
 	{
@@ -1123,7 +1125,7 @@ int32 UMaterialValidationLibrary::GetModifiedShaderCount(UMaterialValidationGrou
 	// If the base material static parameters have changed from those stored on the group then any other material instance 
 	// asset data is no longer valid and we can't iterate the hierarchy without loading the material instances assets.
 	// So instead we approximate the change by only looking at the shader delta on the base material.
-	if (MaterialDesc != nullptr && MaterialDesc->AssetData.StaticPropertyLayoutHash != MaterialAssetData.StaticPropertyLayoutHash)
+	if (MaterialDesc != nullptr && MaterialDesc->AssetData.StaticPropertyLayoutHash != MaterialAssetDataFromObject.StaticPropertyLayoutHash)
 	{
 		const int32 NumShadersBase = MaterialValidationLibrary::GetNumShaderTypes(*InBaseMaterial);
 		return MaterialDesc->NumShadersTotal + (NumShadersBase - MaterialDesc->NumShadersBase) * (1 + MaterialDesc->PermutationHashes.Num());
diff --git a/Engine/Plugins/Editor/MaterialValidation/Source/Private/MaterialValidators.cpp b/Engine/Plugins/Editor/MaterialValidation/Source/Private/MaterialValidators.cpp
--- a/Engine/Plugins/Editor/MaterialValidation/Source/Private/MaterialValidators.cpp
+++ b/Engine/Plugins/Editor/MaterialValidation/Source/Private/MaterialValidators.cpp
@@ -87,7 +87,7 @@ EDataValidationResult UEditorValidator_MaterialPermutation::ValidateLoadedAsset_
 
 		if (NumShaders > 0)
 		{
-			AssetMessage(EMessageSeverity::Info)->AddText(FText::Format(LOCTEXT("ShaderBudget_Info", "Adds {0} new shaders"), FText::AsNumber(NumShaders)));
+			AssetMessage(EMessageSeverity::Info)->AddText(FText::Format(LOCTEXT("ShaderBudget_Info", "New permutation with {0} shaders"), FText::AsNumber(NumShaders)));
 		}
 	}
 
