Material Validation. If we have a new material and MIC with no override then show the shader count cost in the MIC editor.
#rnx
#Synced-CL 48279319

diff --git a/Engine/Plugins/Editor/MaterialValidation/Source/Private/MaterialValidators.cpp b/Engine/Plugins/Editor/MaterialValidation/Source/Private/MaterialValidators.cpp
--- a/Engine/Plugins/Editor/MaterialValidation/Source/Private/MaterialValidators.cpp
+++ b/Engine/Plugins/Editor/MaterialValidation/Source/Private/MaterialValidators.cpp
@@ -162,22 +162,39 @@ bool UMaterialEditorValidator_Permutation::Validate(UMaterialInterface* InMateri
 
 		for (UMaterialValidationGroup const* Group : Groups)
 		{
-			bool bMaterialInGroup = false;
-			bool bMaterialPermutationInGroup = false;
-			UMaterialValidationLibrary::IsMaterialInstanceInGroup(Group, MaterialInstance, bMaterialInGroup, bMaterialPermutationInGroup);
+			bool bIsMaterialInstanceInGroup = false;
+			bool bIsMaterialPermutationInGroup = false;
+			UMaterialValidationLibrary::IsMaterialInstanceInGroup(Group, MaterialInstance, bIsMaterialInstanceInGroup, bIsMaterialPermutationInGroup);
 
-			if (bMaterialInGroup && !bMaterialPermutationInGroup)
+			if (bIsMaterialPermutationInGroup)
 			{
-				UMaterial const* BaseMaterial = MaterialInstance->GetMaterial();
-				const int32 NumShaders = UMaterialEditingLibrary::GetNumShaderTypes(MaterialInstance);
+				continue;
+			}
+
+			UMaterial* BaseMaterial = MaterialInstance->GetMaterial();
 
-				InContext.Message = FText::Format(
-					LOCTEXT("MaterialEditorValidation_Warn", "This material instance creates a new permutation of Base Material '{0}'.\nA permutation of this material is likely to generate {1} unique shaders.\nCan you use an existing permutation instead?"),
-					FText::FromString(BaseMaterial->GetName()),
-					FText::AsNumber(NumShaders));
+			bool bIsBaseMaterialInGroupPath = false;
+			bool bIsBaseMaterialInGroup = false;
+			UMaterialValidationLibrary::IsMaterialInGroup(Group, BaseMaterial, bIsBaseMaterialInGroupPath, bIsBaseMaterialInGroup);
 
-				return false;
+			if (!bIsBaseMaterialInGroupPath)
+			{
+				continue;
 			}
+
+			// bHasStaticPermutation will be false if the material instance has no overrides. Then GetNumShaderTypes() will return 0.
+			// If we get here with that condition then it is because the base material is in the group path but isn't in the group yet.
+			// In that case we can the report the base material shader count.
+			const bool bHasStaticPermutation = MaterialInstance->bHasStaticPermutationResource;
+
+			const int32 NumShaders = bHasStaticPermutation ? UMaterialEditingLibrary::GetNumShaderTypes(MaterialInstance) : UMaterialEditingLibrary::GetNumShaderTypes(BaseMaterial);
+
+			InContext.Message = FText::Format(
+				LOCTEXT("MaterialEditorValidation_Warn", "This material instance creates a new permutation of Base Material '{0}'.\nA permutation of this material is likely to generate {1} unique shaders.\nCan you use an existing permutation instead?"),
+				FText::FromString(BaseMaterial->GetName()),
+				FText::AsNumber(NumShaders));
+
+			return false;
 		}
 	}
 
