Material Validation. Fix bug when tracking overrides of shading model.
The logic for getting shading model is used was wrong for MSM_FromMaterialExpression case.
This led to shader delta overestimates.
#rnx
#Synced-CL 48271239

diff --git a/Engine/Plugins/Editor/MaterialValidation/Source/Private/MaterialValidationLibrary.cpp b/Engine/Plugins/Editor/MaterialValidation/Source/Private/MaterialValidationLibrary.cpp
--- a/Engine/Plugins/Editor/MaterialValidation/Source/Private/MaterialValidationLibrary.cpp
+++ b/Engine/Plugins/Editor/MaterialValidation/Source/Private/MaterialValidationLibrary.cpp
@@ -50,7 +50,7 @@ static FStaticPermutationProperties GetStaticPermutationProperties(UMaterialInte
 	FStaticPermutationProperties StaticProperties;
 	StaticProperties.OpacityMaskClipValue = InMaterial.GetOpacityMaskClipValue();
 	StaticProperties.BlendMode = InMaterial.GetBlendMode();
-	StaticProperties.ShadingModel = InMaterial.GetShadingModels().GetFirstShadingModel();
+	StaticProperties.ShadingModel = InMaterial.IsShadingModelFromMaterialExpression() ? MSM_FromMaterialExpression : InMaterial.GetShadingModels().GetFirstShadingModel();
 	StaticProperties.DitheredLODTransition = InMaterial.IsDitheredLODTransition();
 	StaticProperties.bCastDynamicShadowAsMasked = InMaterial.GetCastDynamicShadowAsMasked();
 	StaticProperties.TwoSided = InMaterial.IsTwoSided();
