Material validation: fix permutation count bugs.
Fixed incorrect shading model logic that did not match behavior where MSM_FromMaterialExpression is never applied as override.
Fixed material instance discovery where detection of found materials was wrong.
#rnx
#Synced-CL 48340728

diff --git a/Engine/Plugins/Editor/MaterialValidation/Source/Private/MaterialValidationLibrary.cpp b/Engine/Plugins/Editor/MaterialValidation/Source/Private/MaterialValidationLibrary.cpp
--- a/Engine/Plugins/Editor/MaterialValidation/Source/Private/MaterialValidationLibrary.cpp
+++ b/Engine/Plugins/Editor/MaterialValidation/Source/Private/MaterialValidationLibrary.cpp
@@ -359,7 +359,8 @@ static void ApplyBasePropertyOverrides(
 	OutProperties.bHasPixelAnimation = InChildOverrideFlags.bOverride_bHasPixelAnimation ? InChildOverrideProperties.bHasPixelAnimation : InParentProperties.bHasPixelAnimation;
 	OutProperties.bEnableTessellation = InChildOverrideFlags.bOverride_bEnableTessellation ? InChildOverrideProperties.bEnableTessellation : InParentProperties.bEnableTessellation;
 	OutProperties.BlendMode = InChildOverrideFlags.bOverride_BlendMode ? InChildOverrideProperties.BlendMode : InParentProperties.BlendMode;
-	OutProperties.ShadingModel = InChildOverrideFlags.bOverride_ShadingModel ? InChildOverrideProperties.ShadingModel : InParentProperties.ShadingModel;
+	// Note that we don't allow override with MSM_FromMaterialExpression.
+	OutProperties.ShadingModel = InChildOverrideFlags.bOverride_ShadingModel && InChildOverrideProperties.ShadingModel != MSM_FromMaterialExpression ? InChildOverrideProperties.ShadingModel : InParentProperties.ShadingModel;
 	OutProperties.OpacityMaskClipValue = InChildOverrideFlags.bOverride_OpacityMaskClipValue ? InChildOverrideProperties.OpacityMaskClipValue : InParentProperties.OpacityMaskClipValue;
 	OutProperties.UsageFlags = (InParentProperties.UsageFlags & ~InChildOverrideFlags.Override_UsageFlags) | (InChildOverrideFlags.Override_UsageFlags & InChildOverrideProperties.UsageFlags);
 }
@@ -553,6 +554,11 @@ static void GetMaterialInstanceHierarchy(UMaterial const& InMaterial, TArray<FNa
 
 		for (const FName& PackageName : ChildPackageNames)
 		{
+			if (OutPackageNames.Contains(PackageName) || NameStack.Contains(PackageName))
+			{
+				continue;
+			}
+
 			TArray<FAssetData> AssetDatas;
 			AssetRegistry.GetAssetsByPackageName(PackageName, AssetDatas, true);
 
@@ -566,10 +572,6 @@ static void GetMaterialInstanceHierarchy(UMaterial const& InMaterial, TArray<FNa
 				{
 					continue;
 				}
-				if (OutPackageNames.Contains(PackageName))
-				{
-					continue;
-				}
 
 				NameStack.Push(PackageName);
 				DepthStack.Push(ParentPackageDepth + 1);
