MaterialValidation: Don't add MICs to validation database when in an excluded directory.
#rnx
#Synced-CL 48343160

diff --git a/Engine/Plugins/Editor/MaterialValidation/Source/Private/MaterialValidationLibrary.cpp b/Engine/Plugins/Editor/MaterialValidation/Source/Private/MaterialValidationLibrary.cpp
--- a/Engine/Plugins/Editor/MaterialValidation/Source/Private/MaterialValidationLibrary.cpp
+++ b/Engine/Plugins/Editor/MaterialValidation/Source/Private/MaterialValidationLibrary.cpp
@@ -258,8 +258,8 @@ static void GetAssetData(
 			const bool Value = bInstanceValid && InstanceMeta.Value.AsStaticSwitch();
 
 			FMaterialParameterMetadata ParentMeta;
-			const bool bBaseValid = ParentMaterial->GetParameterValue(EMaterialParameterType::StaticSwitch, InPermutationData.StaticSwitchParameterInfo[Index], ParentMeta);
-			const bool ParentValue = bBaseValid && ParentMeta.Value.AsStaticSwitch();
+			const bool bParentValid = ParentMaterial->GetParameterValue(EMaterialParameterType::StaticSwitch, InPermutationData.StaticSwitchParameterInfo[Index], ParentMeta);
+			const bool ParentValue = bParentValid && ParentMeta.Value.AsStaticSwitch();
 
 			const uint32 BitArrayIndex = Index / 32;
 			const uint32 BitIndex = Index % 32;
@@ -522,7 +522,7 @@ FString AddPackagePathExtension(FName InPackagePath)
  * Traverse the instance hierarchy for a material and write to a depth first sorted array of package names.
  * The array will have the input material package name as its first entry.
  */
-static void GetMaterialInstanceHierarchy(UMaterial const& InMaterial, TArray<FName>& OutPackageNames, TArray<int32>& OutHierarchyDepths)
+static void GetMaterialInstanceHierarchy(UMaterial const& InMaterial, TConstArrayView<FString> InExcludePaths, TArray<FName>& OutPackageNames, TArray<int32>& OutHierarchyDepths)
 {
 	FAssetRegistryModule& AssetRegistryModule = FModuleManager::LoadModuleChecked<FAssetRegistryModule>("AssetRegistry");
 	IAssetRegistry& AssetRegistry = AssetRegistryModule.Get();
@@ -554,7 +554,7 @@ static void GetMaterialInstanceHierarchy(UMaterial const& InMaterial, TArray<FNa
 
 		for (const FName& PackageName : ChildPackageNames)
 		{
-			if (OutPackageNames.Contains(PackageName) || NameStack.Contains(PackageName))
+			if (OutPackageNames.Contains(PackageName) || NameStack.Contains(PackageName) || IsPackageInPaths(PackageName.ToString(), InExcludePaths))
 			{
 				continue;
 			}
@@ -641,7 +641,7 @@ static int32 GetNumShaderTypes(UMaterialInterface const& InMaterial)
  * This contains asset data for all material instances and information about their unique permutations.
  * Internally this calls LoadMaterialInstances() that is a potentially slow function to load all the material instance assets.
  */
-static void UpdateValidMaterialInstances(UMaterial const& InMaterial, FMaterialValidationDesc& OutMaterialDesc)
+static void UpdateValidMaterialInstances(UMaterial const& InMaterial, TConstArrayView<FString> InExcludePaths, FMaterialValidationDesc& OutMaterialDesc)
 {
 	OutMaterialDesc.MaterialInstanceAssetDatas.Reset();
 	OutMaterialDesc.MaterialInstances.Reset();
@@ -655,7 +655,7 @@ static void UpdateValidMaterialInstances(UMaterial const& InMaterial, FMaterialV
 	// PackageNames will contain InMaterial in the first entry.
 	TArray<FName> PackageNames;
 	TArray<int32> HierarchyDepths;
-	GetMaterialInstanceHierarchy(InMaterial, PackageNames, HierarchyDepths);
+	GetMaterialInstanceHierarchy(InMaterial, InExcludePaths, PackageNames, HierarchyDepths);
 
 	if (PackageNames.Num() <= 1)
 	{
@@ -998,6 +998,10 @@ void UMaterialValidationLibrary::UpdateMaterialPermutationsInGroup(UMaterialVali
 	
 	UE_LOG(LogMaterialValidation, Display, TEXT("Updating %d materials."), InGroup->Materials.Num());
 
+	TArray<FString> SearchPaths;
+	TArray<FString> ExcludePaths;
+	MaterialValidationLibrary::GetSearchPaths(*InGroup, SearchPaths, ExcludePaths);
+
 	FScopedSlowTask SlowTask(static_cast<float>(InGroup->Materials.Num()), LOCTEXT("UpdatingMaterials", "Updating Materials..."));
 	
 	for (TPair<FString, FMaterialValidationDesc>& It : InGroup->Materials)
@@ -1007,7 +1011,7 @@ void UMaterialValidationLibrary::UpdateMaterialPermutationsInGroup(UMaterialVali
 		TSoftObjectPtr<UMaterial> MaterialSoftPtr(MaterialValidationLibrary::ConvertStringPathToSoftObjectPath(It.Key));
 		if (UMaterial const* Material = MaterialSoftPtr.LoadSynchronous())
 		{
-			MaterialValidationLibrary::UpdateValidMaterialInstances(*Material, It.Value);
+			MaterialValidationLibrary::UpdateValidMaterialInstances(*Material, ExcludePaths, It.Value);
 		}
 	}
 
@@ -1102,6 +1106,10 @@ int32 UMaterialValidationLibrary::GetModifiedShaderCount(UMaterialValidationGrou
 		ModifiedObjectPackagePaths.Add(SoftPackagePath.ToString(), Object);
 	}
 
+	TArray<FString> SearchPaths;
+	TArray<FString> ExcludePaths;
+	MaterialValidationLibrary::GetSearchPaths(*InGroup, SearchPaths, ExcludePaths);
+
 	MaterialValidationLibrary::FMaterialPermutationData PermutationData;
 	BuildPermutationData(*InBaseMaterial, PermutationData);
 
@@ -1136,7 +1144,7 @@ int32 UMaterialValidationLibrary::GetModifiedShaderCount(UMaterialValidationGrou
 	// Get the material instance hierarchy from the asset registry.
 	TArray<FName> PackageNames;
 	TArray<int32> HierarchyDepths;
-	MaterialValidationLibrary::GetMaterialInstanceHierarchy(*InBaseMaterial, PackageNames, HierarchyDepths);
+	MaterialValidationLibrary::GetMaterialInstanceHierarchy(*InBaseMaterial, ExcludePaths, PackageNames, HierarchyDepths);
 
 	// Use a stack of asset data to walk the hierarchy with the asset data overrides applied at each level.
 	TArray<FMaterialValidationAssetData_MaterialInstance> AssetDataStack;
