Add table UI for material validation database.
This uses tabs to show a table of materials and material instances.
The material instance table is bare bones for now, but will be extended to include usage flags/static switches etc.
#rnx
#Synced-CL 48711348

diff --git a/Engine/Plugins/Editor/MaterialValidation/Source/MaterialValidation.Build.cs b/Engine/Plugins/Editor/MaterialValidation/Source/MaterialValidation.Build.cs
--- a/Engine/Plugins/Editor/MaterialValidation/Source/MaterialValidation.Build.cs
+++ b/Engine/Plugins/Editor/MaterialValidation/Source/MaterialValidation.Build.cs
@@ -14,8 +14,11 @@ namespace UnrealBuildTool.Rules
 				"DeveloperSettings",
 				"Engine",
 				"Projects",
+				"InputCore",
 				"MaterialEditor",
 				"RHI",
+				"Slate",
+				"SlateCore",
 				"UnrealEd",
 			});
 		}
diff --git a/Engine/Plugins/Editor/MaterialValidation/Source/Private/MaterialValidationGroup.h b/Engine/Plugins/Editor/MaterialValidation/Source/Private/MaterialValidationGroup.h
--- a/Engine/Plugins/Editor/MaterialValidation/Source/Private/MaterialValidationGroup.h
+++ b/Engine/Plugins/Editor/MaterialValidation/Source/Private/MaterialValidationGroup.h
@@ -64,7 +64,7 @@ public:
 	bool bAddProjectPluginPaths;
 
 	/** Map of UMaterial asset paths to validation descriptions. Paths are FString to avoid references to the actual assets which would need managing on asset delete etc. */
-	UPROPERTY(VisibleAnywhere, Category = Material)
+	UPROPERTY()
 	TMap<FString, FMaterialValidationDesc> Materials;
 
 	/** 
diff --git a/Engine/Plugins/Editor/MaterialValidation/Source/Private/MaterialValidationToolkit.cpp b/Engine/Plugins/Editor/MaterialValidation/Source/Private/MaterialValidationToolkit.cpp
new file mode 100644
--- /dev/null
+++ b/Engine/Plugins/Editor/MaterialValidation/Source/Private/MaterialValidationToolkit.cpp
@@ -0,0 +1,676 @@
+// Copyright Epic Games, Inc. All Rights Reserved.
+
+#include "MaterialValidationToolkit.h"
+
+#include "MaterialValidationGroup.h"
+#include "Misc/NotifyHook.h"
+#include "Modules/ModuleManager.h"
+#include "PropertyEditorModule.h"
+#include "Toolkits/AssetEditorToolkit.h"
+#include "UObject/GCObject.h"
+#include "Widgets/Docking/SDockTab.h"
+#include "Widgets/Input/SSearchBox.h"
+#include "Widgets/Text/STextBlock.h"
+#include "Widgets/Views/SHeaderRow.h"
+#include "Widgets/Views/SListView.h"
+
+#define LOCTEXT_NAMESPACE "MaterialValidation"
+
+namespace MaterialValidation {
+
+static const FName ColumnId_AssetPath("AssetPath");
+static const FName ColumnId_MaterialInstanceCount("MaterialInstanceCount");
+static const FName ColumnId_StaticPermutationCount("StaticPermutationCount");
+static const FName ColumnId_ShaderCount("ShaderCount");
+static const FName ColumnId_PermutationHash("PermutationHash");
+static const FName TabId_Details("MainTab");
+static const FName TabId_Materials("MaterialsTab");
+static const FName TabId_MaterialInstances("MaterialInstancesTab");
+
+/** A data table entry describing a UMaterial. */
+struct FMaterialStats
+{
+	FString AssetPath;
+	int32 MaterialInstanceCount = 0;
+	int32 StaticPermutationCount = 0;
+	int32 ShaderCount = 0;
+
+	FMaterialStats(FString const& InPath, int32 InInstanceCount, int32 InPermutationCount, int32 InShaderCount)
+		: AssetPath(InPath)
+		, MaterialInstanceCount(InInstanceCount)
+		, StaticPermutationCount(InPermutationCount)
+		, ShaderCount(InShaderCount)
+	{
+	}
+};
+
+/** A slate container of a data table entry describing a UMaterial. */
+class SMaterialTableRow : public SMultiColumnTableRow<TSharedPtr<FMaterialStats>>
+{
+	TSharedPtr<FMaterialStats> RowItem;
+
+public:
+	SLATE_BEGIN_ARGS(SMaterialTableRow) {}
+		SLATE_ARGUMENT(TSharedPtr<FMaterialStats>, Item)
+	SLATE_END_ARGS()
+
+	void Construct(const FArguments& InArgs, const TSharedRef<STableViewBase>& InOwnerTable)
+	{
+		RowItem = InArgs._Item;
+		SMultiColumnTableRow<TSharedPtr<FMaterialStats>>::Construct(FSuperRowType::FArguments(), InOwnerTable);
+	}
+
+	virtual TSharedRef<SWidget> GenerateWidgetForColumn(const FName& ColumnName) override
+	{
+		if (RowItem.IsValid())
+		{
+			if (ColumnName == ColumnId_AssetPath)
+			{
+				return SNew(STextBlock).Text(FText::FromString(RowItem->AssetPath));
+			}
+			if (ColumnName == ColumnId_MaterialInstanceCount)
+			{
+				return SNew(STextBlock).Text(FText::AsNumber(RowItem->MaterialInstanceCount));
+			}
+			if (ColumnName == ColumnId_StaticPermutationCount)
+			{
+				return SNew(STextBlock).Text(FText::AsNumber(RowItem->StaticPermutationCount));
+			}
+			if (ColumnName == ColumnId_ShaderCount)
+			{
+				return SNew(STextBlock).Text(FText::AsNumber(RowItem->ShaderCount));
+			}
+		}
+		return SNew(STextBlock).Text(FText::FromString("?"));
+	}
+};
+
+/** Delegate for notification of material data row selection. */
+DECLARE_DELEGATE_OneParam(FOnRowItemSelect, TSharedPtr<MaterialValidation::FMaterialStats>);
+
+/** A slate widget containing a table of UMaterial descriptions. */
+class SMaterialTable : public SCompoundWidget
+{
+	/** All the available rows. */
+	TArray<TSharedPtr<FMaterialStats>> AllRows;
+	/** An ordered and filtered subset of AllRows. */
+	TArray<TSharedPtr<FMaterialStats>> FilteredRows;
+	/** A slate view of the table. */
+	TSharedPtr<SListView<TSharedPtr<FMaterialStats>>> ListView;
+
+	/** Search text used for filtering the rows. */
+	FString SearchText;
+	/** Selected column name for sorting the table. */
+	FName SortColumnId = NAME_None;
+	/** Selected mode for sorting the table. */
+	EColumnSortMode::Type SortMode = EColumnSortMode::None;
+
+	/** Delegate for row item selection. */
+	FOnRowItemSelect OnRowItemSelect;
+
+public:
+	SLATE_BEGIN_ARGS(SMaterialTable) {}
+		SLATE_ARGUMENT(TArray<TSharedPtr<FMaterialStats>>, InitialRows)
+		SLATE_EVENT(FOnRowItemSelect, OnRowItemSelect)
+	SLATE_END_ARGS()
+
+	/** Builds the slate table view. */
+	void Construct(const FArguments& InArgs)
+	{
+		AllRows = InArgs._InitialRows;
+		FilteredRows = InArgs._InitialRows;
+
+		SearchText.Reset();
+		SortColumnId = NAME_None;
+		SortMode = EColumnSortMode::None;
+
+		OnRowItemSelect = InArgs._OnRowItemSelect;
+
+		ChildSlot
+		[
+			SNew(SVerticalBox)
+			+ SVerticalBox::Slot().AutoHeight().Padding(2)
+			[
+				SNew(SSearchBox)
+				.OnTextChanged(this, &SMaterialTable::OnSearchTextChanged)
+			]
+			+ SVerticalBox::Slot().FillHeight(1.0f).Padding(2)
+			.FillHeight(1.f)
+			[
+				SAssignNew(ListView, SListView<TSharedPtr<FMaterialStats>>)
+				.HeaderRow
+				(
+					SNew(SHeaderRow)
+					+ SHeaderRow::Column(ColumnId_AssetPath).FillWidth(0.7).HAlignHeader(HAlign_Center).DefaultLabel(FText::FromName(ColumnId_AssetPath)).OnSort(this, &SMaterialTable::OnColumnSortModeChanged)
+					+ SHeaderRow::Column(ColumnId_MaterialInstanceCount).FillWidth(0.1).HAlignHeader(HAlign_Center).HAlignCell(HAlign_Center).DefaultLabel(FText::FromName(ColumnId_MaterialInstanceCount)).OnSort(this, &SMaterialTable::OnColumnSortModeChanged)
+					+ SHeaderRow::Column(ColumnId_StaticPermutationCount).FillWidth(0.1).HAlignHeader(HAlign_Center).HAlignCell(HAlign_Center).DefaultLabel(FText::FromName(ColumnId_StaticPermutationCount)).OnSort(this, &SMaterialTable::OnColumnSortModeChanged)
+					+ SHeaderRow::Column(ColumnId_ShaderCount).FillWidth(0.1).HAlignHeader(HAlign_Center).HAlignCell(HAlign_Center).DefaultLabel(FText::FromName(ColumnId_ShaderCount)).OnSort(this, &SMaterialTable::OnColumnSortModeChanged)
+				)
+				.ListItemsSource(&FilteredRows)
+				.OnGenerateRow_Lambda([](TSharedPtr<FMaterialStats> InItem, const TSharedRef<STableViewBase>& InOwnerTable)
+				{
+					return SNew(SMaterialTableRow, InOwnerTable).Item(InItem);
+				})
+				.SelectionMode(ESelectionMode::Single)
+				.OnMouseButtonDoubleClick(OnRowItemSelect)
+			]
+		];
+	}
+
+	/** Called when modifying the search text. */
+	void OnSearchTextChanged(const FText& InText)
+	{
+		SearchText = InText.ToString();
+		
+		UpdateFilter();
+		UpdateSort();
+		UpdateListView();
+	}
+
+	/** Called when clicking on a column header to sort. */
+	void OnColumnSortModeChanged(const EColumnSortPriority::Type SortPriority, const FName& ColumnId, const EColumnSortMode::Type InSortMode)
+	{
+		if (SortColumnId != ColumnId)
+		{
+			// Initial sort mode is ascending for text, descending for numeric.
+			SortMode = ColumnId == ColumnId_AssetPath ? EColumnSortMode::Ascending : EColumnSortMode::Descending;
+			SortColumnId = ColumnId;
+		}
+		else
+		{
+			SortMode = SortMode == EColumnSortMode::Descending ? EColumnSortMode::Ascending : EColumnSortMode::Descending;
+		}
+
+		UpdateSort();
+		UpdateListView();
+	}
+
+	/**  Apply filter to AllRows to recreate FilteredRows. */
+	void UpdateFilter()
+	{
+		FilteredRows.Reset();
+
+		for (const TSharedPtr<FMaterialStats>& Row : AllRows)
+		{
+			if (SearchText.IsEmpty() ||	Row->AssetPath.Contains(SearchText, ESearchCase::IgnoreCase))
+			{
+				FilteredRows.Add(Row);
+			}
+		}
+	}
+
+	/**  Apply sort to FilteredRows. */
+	void UpdateSort()
+	{
+		const bool bSortModifier = SortMode == EColumnSortMode::Descending;
+
+		if (SortColumnId == ColumnId_AssetPath)
+		{
+			FilteredRows.Sort([bSortModifier](TSharedPtr<FMaterialStats> const& A, TSharedPtr<FMaterialStats> const& B)
+			{
+				return bSortModifier ? (A->AssetPath > B->AssetPath) : (A->AssetPath < B->AssetPath);
+			});
+		}
+		else if (SortColumnId == ColumnId_MaterialInstanceCount)
+		{
+			FilteredRows.Sort([bSortModifier](TSharedPtr<FMaterialStats> const& A, TSharedPtr<FMaterialStats> const& B)
+			{
+				return bSortModifier ? (A->MaterialInstanceCount > B->MaterialInstanceCount) : (A->MaterialInstanceCount < B->MaterialInstanceCount);
+			});
+		}
+		else if (SortColumnId == ColumnId_StaticPermutationCount)
+		{
+			FilteredRows.Sort([bSortModifier](TSharedPtr<FMaterialStats> const& A, TSharedPtr<FMaterialStats> const& B)
+			{
+				return bSortModifier ? (A->StaticPermutationCount > B->StaticPermutationCount) : (A->StaticPermutationCount < B->StaticPermutationCount);
+			});
+		}
+		else if (SortColumnId == ColumnId_ShaderCount)
+		{
+			FilteredRows.Sort([bSortModifier](TSharedPtr<FMaterialStats> const& A, TSharedPtr<FMaterialStats> const& B)
+			{
+				return bSortModifier ? (A->ShaderCount > B->ShaderCount) : (A->ShaderCount < B->ShaderCount);
+			});
+		}
+	}
+
+	/** Update the view. Must be called after any call to UpdateFilter() or UpdateSort(). */
+	void UpdateListView()
+	{
+		if (ListView.IsValid())
+		{
+			ListView->RequestListRefresh();
+		}
+	}
+};
+
+/** A data table entry describing a UMaterialInstance. */
+struct FMaterialInstanceStats
+{
+	FString AssetPath;
+	uint32 PermutationHash = 0;
+
+	FMaterialInstanceStats(FString const& InPath, uint32 InPermutationHash)
+		: AssetPath(InPath)
+		, PermutationHash(InPermutationHash)
+	{
+	}
+};
+
+/** A slate container of a data table entry describing a UMaterialInstance. */
+class SMaterialInstanceTableRow : public SMultiColumnTableRow<TSharedPtr<FMaterialInstanceStats>>
+{
+	TSharedPtr<FMaterialInstanceStats> RowItem;
+
+public:
+	SLATE_BEGIN_ARGS(SMaterialInstanceTableRow) {}
+		SLATE_ARGUMENT(TSharedPtr<FMaterialInstanceStats>, Item)
+	SLATE_END_ARGS()
+
+	void Construct(const FArguments& InArgs, const TSharedRef<STableViewBase>& InOwnerTable)
+	{
+		RowItem = InArgs._Item;
+		SMultiColumnTableRow<TSharedPtr<FMaterialInstanceStats>>::Construct(FSuperRowType::FArguments(), InOwnerTable);
+	}
+
+	virtual TSharedRef<SWidget> GenerateWidgetForColumn(const FName& ColumnName) override
+	{
+		if (RowItem.IsValid())
+		{
+			if (ColumnName == ColumnId_AssetPath)
+			{
+				return SNew(STextBlock).Text(FText::FromString(RowItem->AssetPath));
+			}
+			if (ColumnName == ColumnId_PermutationHash)
+			{
+				return SNew(STextBlock).Text(FText::AsNumber(RowItem->PermutationHash));
+			}
+		}
+		return SNew(STextBlock).Text(FText::FromString("?"));
+	}
+};
+
+/** A slate widget containing a table of UMaterialInstance descriptions. */
+class SMaterialInstanceTable : public SCompoundWidget
+{
+	/** All the available rows. */
+	TArray<TSharedPtr<FMaterialInstanceStats>> AllRows;
+	/** An ordered and filtered subset of AllRows. */
+	TArray<TSharedPtr<FMaterialInstanceStats>> FilteredRows;
+	/** A slate view of the table. */
+	TSharedPtr<SListView<TSharedPtr<FMaterialInstanceStats>>> ListView;
+
+	/** Search text used for filtering the rows. */
+	FString SearchText;
+	/** Selected column name for sorting the table. */
+	FName SortColumnId = NAME_None;
+	/** Selected mode for sorting the table. */
+	EColumnSortMode::Type SortMode = EColumnSortMode::None;
+
+public:
+	SLATE_BEGIN_ARGS(SMaterialInstanceTable) {}
+		SLATE_ARGUMENT(TArray<TSharedPtr<FMaterialInstanceStats>>, InitialRows)
+	SLATE_END_ARGS()
+
+	/** Builds the slate table view. */
+	void Construct(const FArguments& InArgs)
+	{
+		AllRows = InArgs._InitialRows;
+		FilteredRows = InArgs._InitialRows;
+
+		SearchText.Reset();
+		SortColumnId = NAME_None;
+		SortMode = EColumnSortMode::None;
+
+		ChildSlot
+		[
+			SNew(SVerticalBox)
+			+ SVerticalBox::Slot().AutoHeight().Padding(2)
+			[
+				SNew(SSearchBox)
+				.OnTextChanged(this, &SMaterialInstanceTable::OnSearchTextChanged)
+			]
+			+ SVerticalBox::Slot().FillHeight(1.0f).Padding(2)
+			.FillHeight(1.f)
+			[
+				SAssignNew(ListView, SListView<TSharedPtr<FMaterialInstanceStats>>)
+				.HeaderRow
+				(
+					SNew(SHeaderRow)
+					+ SHeaderRow::Column(ColumnId_AssetPath).FillWidth(0.7).HAlignHeader(HAlign_Center).DefaultLabel(FText::FromName(ColumnId_AssetPath)).OnSort(this, &SMaterialInstanceTable::OnColumnSortModeChanged)
+					+ SHeaderRow::Column(ColumnId_PermutationHash).FillWidth(0.3).HAlignHeader(HAlign_Center).HAlignCell(HAlign_Center).DefaultLabel(FText::FromName(ColumnId_PermutationHash)).OnSort(this, &SMaterialInstanceTable::OnColumnSortModeChanged)
+				)
+				.ListItemsSource(&FilteredRows)
+				.OnGenerateRow_Lambda([](TSharedPtr<FMaterialInstanceStats> InItem, const TSharedRef<STableViewBase>& InOwnerTable)
+				{
+					return SNew(SMaterialInstanceTableRow, InOwnerTable).Item(InItem);
+				})
+			]
+		];
+	}
+
+	/** Set the table content. */
+	void SetRows(TArray<TSharedPtr<FMaterialInstanceStats>> const& InRows)
+	{
+		AllRows = InRows;
+
+		UpdateFilter();
+		UpdateSort();
+		UpdateListView();
+	}
+
+	/** Called when modifying the search text. */
+	void OnSearchTextChanged(const FText& InText)
+	{
+		SearchText = InText.ToString();
+
+		UpdateFilter();
+		UpdateSort();
+		UpdateListView();
+	}
+
+	/** Called when clicking on a column header to sort. */
+	void OnColumnSortModeChanged(const EColumnSortPriority::Type SortPriority, const FName& ColumnId, const EColumnSortMode::Type InSortMode)
+	{
+		if (SortColumnId != ColumnId)
+		{
+			// Initial sort mode is ascending for text, descending for numeric.
+			SortMode = ColumnId == ColumnId_AssetPath ? EColumnSortMode::Ascending : EColumnSortMode::Descending;
+			SortColumnId = ColumnId;
+		}
+		else
+		{
+			SortMode = SortMode == EColumnSortMode::Descending ? EColumnSortMode::Ascending : EColumnSortMode::Descending;
+		}
+
+		UpdateSort();
+		UpdateListView();
+	}
+
+	/**  Apply filter to AllRows to recreate FilteredRows. */
+	void UpdateFilter()
+	{
+		FilteredRows.Reset();
+
+		for (const TSharedPtr<FMaterialInstanceStats>& Row : AllRows)
+		{
+			if (SearchText.IsEmpty() || Row->AssetPath.Contains(SearchText, ESearchCase::IgnoreCase))
+			{
+				FilteredRows.Add(Row);
+			}
+		}
+	}
+
+	/**  Apply sort to FilteredRows. */
+	void UpdateSort()
+	{
+		const bool bSortModifier = SortMode == EColumnSortMode::Descending;
+
+		if (SortColumnId == ColumnId_AssetPath)
+		{
+			FilteredRows.Sort([bSortModifier](TSharedPtr<FMaterialInstanceStats> const& A, TSharedPtr<FMaterialInstanceStats> const& B)
+			{
+				return bSortModifier ? (A->AssetPath > B->AssetPath) : (A->AssetPath < B->AssetPath);
+			});
+		}
+		else if (SortColumnId == ColumnId_PermutationHash)
+		{
+			FilteredRows.Sort([bSortModifier](TSharedPtr<FMaterialInstanceStats> const& A, TSharedPtr<FMaterialInstanceStats> const& B)
+			{
+				return bSortModifier ? (A->PermutationHash > B->PermutationHash) : (A->PermutationHash < B->PermutationHash);
+			});
+		}
+	}
+
+	/** Update the view. Must be called after any call to UpdateFilter() or UpdateSort(). */
+	void UpdateListView()
+	{
+		if (ListView.IsValid())
+		{
+			ListView->RequestListRefresh();
+		}
+	}
+};
+
+static TSharedPtr<SMaterialTable> CreateMaterialTable(UMaterialValidationGroup const* InGroup, FOnRowItemSelect OnRowItemSelect)
+{
+	TArray<TSharedPtr<FMaterialStats>> Stats;
+	if (InGroup != nullptr)
+	{
+		for (TPair<FString, FMaterialValidationDesc> const& It : InGroup->Materials)
+		{
+			Stats.Add(MakeShared<FMaterialStats>(It.Key, It.Value.MaterialInstances.Num(), It.Value.PermutationHashes.Num(), It.Value.NumShadersTotal));
+		}
+	}
+	return SNew(SMaterialTable).InitialRows(Stats).OnRowItemSelect(OnRowItemSelect);
+}
+
+static TSharedPtr<SMaterialInstanceTable> CreateMaterialInstanceTable(FMaterialValidationDesc const* InMaterialDesc)
+{
+	TArray<TSharedPtr<FMaterialInstanceStats>> Stats;
+	if (InMaterialDesc != nullptr)
+	{
+		for (TPair<FString, int32> const& It : InMaterialDesc->MaterialInstances)
+		{
+			const int32 AssetDataIndex = It.Value;
+			const uint32 PermutationHash = InMaterialDesc->MaterialInstanceAssetDatas.IsValidIndex(AssetDataIndex) ? InMaterialDesc->MaterialInstanceAssetDatas[AssetDataIndex].PermutationHash : 0;
+			Stats.Add(MakeShared<FMaterialInstanceStats>(It.Key, PermutationHash));
+		}
+	}
+	return SNew(SMaterialInstanceTable).InitialRows(Stats);
+}
+
+static void UpdateMaterialInstanceTable(TSharedPtr<SMaterialInstanceTable> InTable, FMaterialValidationDesc const* InMaterialDesc)
+{
+	TArray<TSharedPtr<MaterialValidation::FMaterialInstanceStats>> Stats;
+	if (InMaterialDesc != nullptr)
+	{
+		for (TPair<FString, int32> const& It : InMaterialDesc->MaterialInstances)
+		{
+			const int32 AssetDataIndex = It.Value;
+			const uint32 PermutationHash = InMaterialDesc->MaterialInstanceAssetDatas.IsValidIndex(AssetDataIndex) ? InMaterialDesc->MaterialInstanceAssetDatas[AssetDataIndex].PermutationHash : 0;
+			Stats.Add(MakeShared<MaterialValidation::FMaterialInstanceStats>(It.Key, PermutationHash));
+		}
+	}
+	InTable->SetRows(Stats);
+}
+
+} // namespace MaterialValidation
+
+/** 
+ * AssetEditorToolkit for UMaterialValidationGroup. 
+ * This sets up the tabs for showing various table views of the data in the group.
+ */
+class FMaterialValidationToolkit : public FAssetEditorToolkit, public FGCObject, public FNotifyHook
+{
+	/** The current UMaterialValidationGroup being edited. */
+	TObjectPtr<UMaterialValidationGroup> EditGroup;
+
+	/** Tab manager for the toolkit. */
+	TSharedPtr<class FTabManager> TabManager;
+	/** Menu category for the toolkit tabs. */
+	TSharedPtr<class FWorkspaceItem> WorkspaceMenuCategory;
+
+	/** Standard details panel view. */
+	TSharedPtr<class IDetailsView> DetailsView;
+	/** Material table widget. */
+	TSharedPtr<MaterialValidation::SMaterialTable> MaterialTable;
+	/** Material instance table widget. */
+	TSharedPtr<MaterialValidation::SMaterialInstanceTable> MaterialInstanceTable;
+
+public:
+	// Begin FGCObject overrides.
+	virtual FString GetReferencerName() const override
+	{
+		return TEXT("FMaterialValidationToolkit");
+	}
+
+	virtual void AddReferencedObjects(FReferenceCollector& Collector) override
+	{
+		Collector.AddReferencedObject(EditGroup);
+	}
+	// End FGCObject overrides.
+
+	// Begin FAssetEditorToolkit overrides.
+	virtual FName GetToolkitFName() const override
+	{
+		return FName("MaterialValidation");
+	}
+
+	virtual FText GetBaseToolkitName() const override
+	{
+		return LOCTEXT("Toolkit_AppLabel", "MaterialValidation");
+	}
+
+	virtual FString GetWorldCentricTabPrefix() const override
+	{
+		return LOCTEXT("Toolkit_TabPrefix", "MaterialValidation").ToString();
+	}
+
+	virtual FLinearColor GetWorldCentricTabColorScale() const override
+	{
+		return FLinearColor(0.0f, 0.0f, 0.2f, 0.5f);
+	}
+
+	virtual void RegisterTabSpawners(const TSharedRef<class FTabManager>& InTabManager) override
+	{
+		WorkspaceMenuCategory = InTabManager->AddLocalWorkspaceMenuCategory(LOCTEXT("Toolkit_WorkspaceMenuCatgory", "Material Validation"));
+
+		InTabManager->RegisterTabSpawner(MaterialValidation::TabId_Details, FOnSpawnTab::CreateSP(this, &FMaterialValidationToolkit::SpawnDetailsTab))
+			.SetDisplayName(LOCTEXT("Toolkit_DetailsTab_DisplayName", "Details"))
+			.SetGroup(WorkspaceMenuCategory.ToSharedRef())
+			.SetIcon(FSlateIcon(FAppStyle::GetAppStyleSetName(), "LevelEditor.Tabs.Details"));
+
+		InTabManager->RegisterTabSpawner(MaterialValidation::TabId_Materials, FOnSpawnTab::CreateSP(this, &FMaterialValidationToolkit::SpawnMaterialsTab))
+			.SetDisplayName(LOCTEXT("Toolkit_MaterialsTab_DisplayName", "Materials"))
+			.SetGroup(WorkspaceMenuCategory.ToSharedRef())
+			.SetIcon(FSlateIcon(FAppStyle::GetAppStyleSetName(), "LevelEditor.Tabs.StatsViewer"));
+
+		InTabManager->RegisterTabSpawner(MaterialValidation::TabId_MaterialInstances, FOnSpawnTab::CreateSP(this, &FMaterialValidationToolkit::SpawnMaterialInstancesTab))
+			.SetDisplayName(LOCTEXT("Toolkit_MaterialInstancesTab_DisplayName", "Material Instances"))
+			.SetGroup(WorkspaceMenuCategory.ToSharedRef())
+			.SetIcon(FSlateIcon(FAppStyle::GetAppStyleSetName(), "LevelEditor.Tabs.StatsViewer"));
+	}
+
+	virtual void UnregisterTabSpawners(const TSharedRef<class FTabManager>& InTabManager) override
+	{
+		FAssetEditorToolkit::UnregisterTabSpawners(InTabManager);
+		InTabManager->UnregisterTabSpawner(MaterialValidation::TabId_Details);
+		InTabManager->UnregisterTabSpawner(MaterialValidation::TabId_Materials);
+		InTabManager->UnregisterTabSpawner(MaterialValidation::TabId_MaterialInstances);
+	}
+	// End FAssetEditorToolkit overrides.
+
+	TSharedRef<SDockTab> SpawnDetailsTab(const FSpawnTabArgs& Args)
+	{
+		return SNew(SDockTab)
+			.Label(LOCTEXT("Toolkit_DetailsTab_Label", "Details"))
+			[
+				DetailsView.ToSharedRef()
+			];
+	}
+
+	TSharedRef<SDockTab> SpawnMaterialsTab(const FSpawnTabArgs& Args)
+	{
+		return SNew(SDockTab)
+			.Label(LOCTEXT("Toolkit_Materials_Label", "Materials"))
+			[
+				MaterialTable.ToSharedRef()
+			];
+	}
+
+	TSharedRef<SDockTab> SpawnMaterialInstancesTab(const FSpawnTabArgs& Args)
+	{
+		return SNew(SDockTab)
+			.Label(LOCTEXT("Toolkit_MaterialInstances_Label", "Material Instances"))
+			[
+				MaterialInstanceTable.ToSharedRef()
+			];
+	}
+
+	/** Populate the toolkit layout. */
+	void Initialize(const EToolkitMode::Type Mode, const TSharedPtr< class IToolkitHost >& InitToolkitHost, UMaterialValidationGroup* InGroup)
+	{
+		EditGroup = InGroup;
+
+		FPropertyEditorModule& PropertyEditorModule = FModuleManager::LoadModuleChecked<FPropertyEditorModule>("PropertyEditor");
+
+		FDetailsViewArgs Args;
+		Args.bHideSelectionTip = true;
+		Args.NotifyHook = this;
+		DetailsView = PropertyEditorModule.CreateDetailView(Args);
+		DetailsView->SetObject(EditGroup);
+
+		MaterialTable = MaterialValidation::CreateMaterialTable(EditGroup, MaterialValidation::FOnRowItemSelect::CreateSP(this, &FMaterialValidationToolkit::OnRowItemSelect));
+
+		MaterialInstanceTable = MaterialValidation::CreateMaterialInstanceTable(nullptr);
+
+		const TSharedRef<FTabManager::FLayout> Layout = FTabManager::NewLayout("MaterialValidationGroupLayout_v3")
+			->AddArea
+			(
+				FTabManager::NewPrimaryArea()->SetOrientation(Orient_Horizontal)
+				->Split
+				(
+					FTabManager::NewSplitter()->SetOrientation(Orient_Vertical)
+					->Split(FTabManager::NewStack()
+						->AddTab(MaterialValidation::TabId_Details, ETabState::OpenedTab)
+						->SetForegroundTab(MaterialValidation::TabId_Details))
+					->Split(FTabManager::NewStack()
+						->AddTab(MaterialValidation::TabId_Materials, ETabState::OpenedTab))
+					->Split(FTabManager::NewStack()
+						->AddTab(MaterialValidation::TabId_MaterialInstances, ETabState::ClosedTab))
+				)
+			);
+
+		InitAssetEditor(Mode, InitToolkitHost, GetToolkitFName(), Layout, true /*bCreateDefaultStandaloneMenu*/, false /*bCreateDefaultToolbar*/, EditGroup);
+	}
+
+	/** Apply selected material to the material instances table. */
+	void OnRowItemSelect(TSharedPtr<MaterialValidation::FMaterialStats> InMaterialStats)
+	{
+		FMaterialValidationDesc const* MaterialDesc = nullptr;
+		if (EditGroup != nullptr && InMaterialStats.IsValid())
+		{
+			MaterialDesc = EditGroup->Materials.Find(InMaterialStats->AssetPath);
+		}
+
+		MaterialValidation::UpdateMaterialInstanceTable(MaterialInstanceTable, MaterialDesc);
+
+		InvokeTab(MaterialValidation::TabId_MaterialInstances);
+	}
+};
+
+FText UAssetDefinition_MaterialValidationGroup::GetAssetDisplayName() const 
+{
+	return LOCTEXT("Asset_DisplayName", "Material Validation Group"); 
+}
+
+FLinearColor UAssetDefinition_MaterialValidationGroup::GetAssetColor() const 
+{
+	return FLinearColor::Gray; 
+}
+
+TSoftClassPtr<UObject> UAssetDefinition_MaterialValidationGroup::GetAssetClass() const 
+{
+	return UMaterialValidationGroup::StaticClass(); 
+}
+
+TConstArrayView<FAssetCategoryPath> UAssetDefinition_MaterialValidationGroup::GetAssetCategories() const
+{
+	static TArray<FAssetCategoryPath> Categories = { EAssetCategoryPaths::Material };
+	return Categories;
+}
+
+EAssetCommandResult UAssetDefinition_MaterialValidationGroup::OpenAssets(const FAssetOpenArgs& OpenArgs) const
+{
+	for (UMaterialValidationGroup* Group : OpenArgs.LoadObjects<UMaterialValidationGroup>())
+	{
+		TSharedRef<FMaterialValidationToolkit> NewToolkit(new FMaterialValidationToolkit());
+		NewToolkit->Initialize(OpenArgs.GetToolkitMode(), OpenArgs.ToolkitHost, Group);
+	}
+
+	return EAssetCommandResult::Handled;
+
+}
+
+#undef LOCTEXT_NAMESPACE
diff --git a/Engine/Plugins/Editor/MaterialValidation/Source/Private/MaterialValidationToolkit.h b/Engine/Plugins/Editor/MaterialValidation/Source/Private/MaterialValidationToolkit.h
new file mode 100644
--- /dev/null
+++ b/Engine/Plugins/Editor/MaterialValidation/Source/Private/MaterialValidationToolkit.h
@@ -0,0 +1,19 @@
+// Copyright Epic Games, Inc. All Rights Reserved.
+
+#pragma once
+
+#include "AssetDefinitionDefault.h"
+#include "MaterialValidationToolkit.generated.h"
+
+/** UMaterialValidationGroup asset definition including editor toolkit setup. */
+UCLASS()
+class UAssetDefinition_MaterialValidationGroup : public UAssetDefinitionDefault
+{
+	GENERATED_BODY()
+
+	virtual FText GetAssetDisplayName() const override;
+	virtual FLinearColor GetAssetColor() const override;
+	virtual TSoftClassPtr<UObject> GetAssetClass() const override;
+	virtual TConstArrayView<FAssetCategoryPath> GetAssetCategories() const override;
+	virtual EAssetCommandResult OpenAssets(const FAssetOpenArgs& OpenArgs) const override;
+};
