Fix case where validation warning was shown for material instance without a static permutation

#rnx

#Synced-CL 49018876

diff --git a/Engine/Plugins/Editor/MaterialValidation/Source/Private/MaterialValidators.cpp b/Engine/Plugins/Editor/MaterialValidation/Source/Private/MaterialValidators.cpp
--- a/Engine/Plugins/Editor/MaterialValidation/Source/Private/MaterialValidators.cpp
+++ b/Engine/Plugins/Editor/MaterialValidation/Source/Private/MaterialValidators.cpp
@@ -65,7 +65,7 @@ EDataValidationResult UEditorValidator_MaterialPermutation::ValidateLoadedAsset_
 				}
 			}
 			
-			if (MaterialInstance)
+			if (MaterialInstance && MaterialInstance->bHasStaticPermutationResource)
 			{
 				bool bMaterialInGroup = false;
 				bool bMaterialPermutationInGroup = false;
@@ -192,11 +192,16 @@ bool UMaterialEditorValidator_Permutation::Validate(UMaterialInterface* InMateri
 				continue;
 			}
 
-			// bHasStaticPermutation will be false if the material instance has no overrides. Then GetNumShaderTypes() will return 0.
-			// If we get here with that condition then it is because the base material is in the group path but isn't in the group yet.
-			// In that case we can the report the base material shader count.
+			// bHasStaticPermutation will be false if the material instance has no overrides.
+			// If we get here with that condition and the base material is in the group then there are no shaders to count.
 			const bool bHasStaticPermutation = MaterialInstance->bHasStaticPermutationResource;
+			if (!bHasStaticPermutation && bIsBaseMaterialInGroup)
+			{
+				continue;
+			}
 
+			// If there is no static permutation then the base material is in the group path but isn't in the group yet.
+			// In that case report the shader count from the base material shader count.
 			const int32 NumShaders = bHasStaticPermutation ? UMaterialEditingLibrary::GetNumShaderTypes(MaterialInstance) : UMaterialEditingLibrary::GetNumShaderTypes(BaseMaterial);
 
 			InContext.Message = FText::Format(
