Fix uninitialized USTRUCT members
#rnx
#Synced-CL 49559171

diff --git a/Engine/Plugins/Editor/MaterialValidation/Source/Private/MaterialValidationGroup.h b/Engine/Plugins/Editor/MaterialValidation/Source/Private/MaterialValidationGroup.h
--- a/Engine/Plugins/Editor/MaterialValidation/Source/Private/MaterialValidationGroup.h
+++ b/Engine/Plugins/Editor/MaterialValidation/Source/Private/MaterialValidationGroup.h
@@ -33,11 +33,11 @@ struct FMaterialValidationDesc
 
 	/** Number of shader types on the base material permutation. */
 	UPROPERTY()
-	int32 NumShadersBase;
+	int32 NumShadersBase = 0;
 
 	/** Total sum of shader types for all permutations found of this material. */
 	UPROPERTY(VisibleAnywhere, Category = Material)
-	int32 NumShadersTotal;
+	int32 NumShadersTotal = 0;
 };
 
 /** Data Asset containing the validation setup for a group of UMaterials. */
