Extend material validation toolkit UI to show a table of information about material instances for a selected material.
#rnx
#Synced-CL 50647859

diff --git a/Engine/Plugins/Editor/MaterialValidation/Source/Private/MaterialValidationLibrary.cpp b/Engine/Plugins/Editor/MaterialValidation/Source/Private/MaterialValidationLibrary.cpp
--- a/Engine/Plugins/Editor/MaterialValidation/Source/Private/MaterialValidationLibrary.cpp
+++ b/Engine/Plugins/Editor/MaterialValidation/Source/Private/MaterialValidationLibrary.cpp
@@ -1172,7 +1172,7 @@ int32 UMaterialValidationLibrary::GetModifiedShaderCount(UMaterialValidationGrou
 
 		FMaterialValidationAssetData_MaterialInstance AssetData;
 		UMaterialInstanceConstant* MaterialInstance = nullptr;
-		const bool bValidAssetData = GetMaterialInstanceAssetData(PackagePath, ModifiedObjectPackagePaths, MaterialDesc, PermutationData, AssetData, MaterialInstance);
+		const bool bValidAssetData = MaterialValidationLibrary::GetMaterialInstanceAssetData(PackagePath, ModifiedObjectPackagePaths, MaterialDesc, PermutationData, AssetData, MaterialInstance);
 
 		if (!bValidAssetData)
 		{
@@ -1214,4 +1214,126 @@ int32 UMaterialValidationLibrary::GetModifiedShaderCount(UMaterialValidationGrou
 	return NumShaders;
 }
 
+/** 
+ * Note there is quite a bit of duplication with GetModifiedShaderCount() in this function. 
+ * It might be good to consolidate on a single material instance visitor that can do arbitrary work?
+ */
+void UMaterialValidationLibrary::GetMaterialInstanceAssetDatas(
+	UMaterialValidationGroup const* InGroup, 
+	FString const& MaterialPath,
+	TArray<FString>& OutMaterialInstancePaths, 
+	TArray<FMaterialValidationAssetData_MaterialInstance>& OutMaterialInstanceAssetDatas,
+	TArray<FName>& OutStaticSwitchNames,
+	TArray<FName>& OutComponentMaskNames)
+{
+	if (InGroup == nullptr)
+	{
+		return;
+	}
+
+	FMaterialValidationDesc const* MaterialDesc = InGroup->Materials.Find(MaterialPath);
+	if (MaterialDesc == nullptr)
+	{
+		return;
+	}
+
+	TSoftObjectPtr<UMaterial> MaterialSoftPtr(MaterialValidationLibrary::ConvertStringPathToSoftObjectPath(MaterialPath));
+	UMaterial const* Material = MaterialSoftPtr.LoadSynchronous();
+	if (Material == nullptr)
+	{
+		return;
+	}
+
+	TArray<FString> SearchPaths;
+	TArray<FString> ExcludePaths;
+	MaterialValidationLibrary::GetSearchPaths(*InGroup, SearchPaths, ExcludePaths);
+
+	MaterialValidationLibrary::FMaterialPermutationData PermutationData;
+	BuildPermutationData(*Material, PermutationData);
+
+	for (FMaterialParameterInfo const& Info : PermutationData.StaticSwitchParameterInfo)
+	{
+		OutStaticSwitchNames.Add(Info.Name);
+	}
+	for (FMaterialParameterInfo const& Info : PermutationData.ComponentMaskParameterInfo)
+	{
+		OutComponentMaskNames.Add(Info.Name);
+	}
+
+	FMaterialValidationAssetData_Material MaterialAssetData = MaterialDesc->AssetData;
+
+	// Get the material instance hierarchy from the asset registry.
+	TArray<FName> PackageNames;
+	TArray<int32> HierarchyDepths;
+	MaterialValidationLibrary::GetMaterialInstanceHierarchy(*Material, ExcludePaths, PackageNames, HierarchyDepths);
+
+	// Use a stack of asset data to walk the hierarchy with the asset data overrides applied at each level.
+	TArray<FMaterialValidationAssetData_MaterialInstance> AssetDataStack;
+
+	// Build the root asset data from the base material.
+	{
+		FMaterialValidationAssetData_MaterialInstance& Root = AssetDataStack.AddDefaulted_GetRef();
+		Root.StaticSwitchOverrideValues = MaterialAssetData.StaticSwitchValues;
+		Root.ComponentMaskOverrideValues = MaterialAssetData.ComponentMaskValues;
+		Root.StaticProperties = MaterialAssetData.StaticProperties;
+		Root.MaterialLayerHash = MaterialAssetData.MaterialLayerHash;
+		Root.PermutationHash = MaterialAssetData.PermutationHash;
+	}
+
+	// Walk the material instance hierarchy.
+	for (int32 Index = 1; Index < PackageNames.Num(); ++Index)
+	{
+		const int32 CurrentDepth = HierarchyDepths[Index];
+		FString PackagePath = MaterialValidationLibrary::AddPackagePathExtension(PackageNames[Index]);
+
+		bool bValidAssetData = false;
+		FMaterialValidationAssetData_MaterialInstance AssetData;
+		if (int32 const* AssetDataIndex = MaterialDesc->MaterialInstances.Find(PackagePath))
+		{
+			if (*AssetDataIndex == INDEX_NONE)
+			{
+				// Set up a default asset data with the right sizes.
+				const int32 StaticSwitchArrayNum = MaterialDesc->AssetData.StaticSwitchValues.Num();
+				AssetData.StaticSwitchOverrideMask.SetNumZeroed(StaticSwitchArrayNum);
+				AssetData.StaticSwitchOverrideValues.SetNumZeroed(StaticSwitchArrayNum);
+				const int32 ComponentMaskArrayNum = MaterialDesc->AssetData.ComponentMaskValues.Num();
+				AssetData.ComponentMaskOverrideMask.SetNumZeroed(ComponentMaskArrayNum);
+				AssetData.ComponentMaskOverrideValues.SetNumZeroed(ComponentMaskArrayNum);
+				bValidAssetData = true;
+			}
+			else if (MaterialDesc->MaterialInstanceAssetDatas.IsValidIndex(*AssetDataIndex))
+			{
+				AssetData = MaterialDesc->MaterialInstanceAssetDatas[*AssetDataIndex];
+				bValidAssetData = true;
+			}
+		}
+		
+		if (!bValidAssetData)
+		{
+			// Ignore material instance and it's children.
+			while (Index < PackageNames.Num() - 1 && HierarchyDepths[Index + 1] > CurrentDepth)
+			{
+				++Index;
+			}
+			continue;
+		}
+
+		// Apply the current asset data overrides to the parent to get the final asset data for the child.
+		AssetDataStack.SetNum(CurrentDepth + 1);
+		FMaterialValidationAssetData_MaterialInstance& Child = AssetDataStack[CurrentDepth];
+		FMaterialValidationAssetData_MaterialInstance const& Parent = AssetDataStack[CurrentDepth - 1];
+
+		Child.StaticSwitchOverrideValues.SetNumZeroed(Parent.StaticSwitchOverrideValues.Num());
+		Child.ComponentMaskOverrideValues.SetNumZeroed(Parent.ComponentMaskOverrideValues.Num());
+		MaterialValidationLibrary::ApplyBitMaskOverrides(Parent.StaticSwitchOverrideValues, AssetData.StaticSwitchOverrideMask, AssetData.StaticSwitchOverrideValues, Child.StaticSwitchOverrideValues);
+		MaterialValidationLibrary::ApplyBitMaskOverrides(Parent.ComponentMaskOverrideValues, AssetData.ComponentMaskOverrideMask, AssetData.ComponentMaskOverrideValues, Child.ComponentMaskOverrideValues);
+		MaterialValidationLibrary::ApplyBasePropertyOverrides(Parent.StaticProperties, AssetData.StaticPropertyOverrideFlags, AssetData.StaticProperties, Child.StaticProperties);
+		Child.MaterialLayerHash = AssetData.MaterialLayerHash;
+		Child.PermutationHash = MaterialValidationLibrary::BuildPermutationHash(Child.StaticProperties, Child.StaticSwitchOverrideValues, Child.ComponentMaskOverrideValues, Child.MaterialLayerHash);
+
+		OutMaterialInstancePaths.Add(PackagePath);
+		OutMaterialInstanceAssetDatas.Add(Child);
+	}
+}
+
 #undef LOCTEXT_NAMESPACE
diff --git a/Engine/Plugins/Editor/MaterialValidation/Source/Private/MaterialValidationLibrary.h b/Engine/Plugins/Editor/MaterialValidation/Source/Private/MaterialValidationLibrary.h
--- a/Engine/Plugins/Editor/MaterialValidation/Source/Private/MaterialValidationLibrary.h
+++ b/Engine/Plugins/Editor/MaterialValidation/Source/Private/MaterialValidationLibrary.h
@@ -89,6 +89,7 @@ class UMaterialValidationLibrary : public UBlueprintFunctionLibrary
 	 *
 	 * @returns	The total number of shaders calculated at the last time we updated the group.
 	 */
+	UFUNCTION(BlueprintCallable, Category = Material)
 	static int32 GetShaderCount(UMaterialValidationGroup const* InGroup, UMaterial* InMaterial);
 
 	/**
@@ -102,4 +103,13 @@ class UMaterialValidationLibrary : public UBlueprintFunctionLibrary
 	 */
 	UFUNCTION(BlueprintCallable, Category = Material)
 	static int32 GetModifiedShaderCount(UMaterialValidationGroup const* InGroup, UMaterial* InMaterial, TArray<UMaterialInterface*>& InModifiedObjects);
+
+	/** Get all of the material instance data for a single material. */
+	static void GetMaterialInstanceAssetDatas(
+		UMaterialValidationGroup const* InGroup,
+		FString const& MaterialPath,
+		TArray<FString>& OutMaterialInstancePaths,
+		TArray<struct FMaterialValidationAssetData_MaterialInstance>& OutMaterialInstanceAssetDatas,
+		TArray<FName>& OutStaticSwitchNames,
+		TArray<FName>& OutComponentMaskNames);
 };
diff --git a/Engine/Plugins/Editor/MaterialValidation/Source/Private/MaterialValidationToolkit.cpp b/Engine/Plugins/Editor/MaterialValidation/Source/Private/MaterialValidationToolkit.cpp
--- a/Engine/Plugins/Editor/MaterialValidation/Source/Private/MaterialValidationToolkit.cpp
+++ b/Engine/Plugins/Editor/MaterialValidation/Source/Private/MaterialValidationToolkit.cpp
@@ -2,14 +2,18 @@
 
 #include "MaterialValidationToolkit.h"
 
+#include "Materials/MaterialInterface.h"
 #include "MaterialValidationGroup.h"
+#include "MaterialValidationLibrary.h"
 #include "Misc/NotifyHook.h"
 #include "Modules/ModuleManager.h"
 #include "PropertyEditorModule.h"
 #include "Toolkits/AssetEditorToolkit.h"
 #include "UObject/GCObject.h"
 #include "Widgets/Docking/SDockTab.h"
+#include "Widgets/Input/SComboButton.h"
 #include "Widgets/Input/SSearchBox.h"
+#include "Widgets/Layout/SScrollBox.h"
 #include "Widgets/Text/STextBlock.h"
 #include "Widgets/Views/SHeaderRow.h"
 #include "Widgets/Views/SListView.h"
@@ -18,14 +22,8 @@
 
 namespace MaterialValidation {
 
-static const FName ColumnId_AssetPath("AssetPath");
-static const FName ColumnId_MaterialInstanceCount("MaterialInstanceCount");
-static const FName ColumnId_StaticPermutationCount("StaticPermutationCount");
-static const FName ColumnId_ShaderCount("ShaderCount");
-static const FName ColumnId_PermutationHash("PermutationHash");
-static const FName TabId_Details("MainTab");
-static const FName TabId_Materials("MaterialsTab");
-static const FName TabId_MaterialInstances("MaterialInstancesTab");
+/** Global shared value to track user selected width of asset name columns. */
+static float GAssetNameColumnWidth = 700.f;
 
 /** A data table entry describing a UMaterial. */
 struct FMaterialStats
@@ -35,6 +33,12 @@ struct FMaterialStats
 	int32 StaticPermutationCount = 0;
 	int32 ShaderCount = 0;
 
+	// Stat table columns are referenced by FName.
+	static const FName ColumnId_AssetPath;
+	static const FName ColumnId_MaterialInstanceCount;
+	static const FName ColumnId_StaticPermutationCount;
+	static const FName ColumnId_ShaderCount;
+
 	FMaterialStats(FString const& InPath, int32 InInstanceCount, int32 InPermutationCount, int32 InShaderCount)
 		: AssetPath(InPath)
 		, MaterialInstanceCount(InInstanceCount)
@@ -44,9 +48,15 @@ struct FMaterialStats
 	}
 };
 
+const FName FMaterialStats::ColumnId_AssetPath("CID_AssetPath");
+const FName FMaterialStats::ColumnId_MaterialInstanceCount("CID_MaterialInstanceCount");
+const FName FMaterialStats::ColumnId_StaticPermutationCount("CID_StaticPermutationCount");
+const FName FMaterialStats::ColumnId_ShaderCount("CID_ShaderCount");
+
 /** A slate container of a data table entry describing a UMaterial. */
 class SMaterialTableRow : public SMultiColumnTableRow<TSharedPtr<FMaterialStats>>
 {
+	/** The material stats for the row. */
 	TSharedPtr<FMaterialStats> RowItem;
 
 public:
@@ -54,6 +64,7 @@ public:
 		SLATE_ARGUMENT(TSharedPtr<FMaterialStats>, Item)
 	SLATE_END_ARGS()
 
+	/** Builds the row. */
 	void Construct(const FArguments& InArgs, const TSharedRef<STableViewBase>& InOwnerTable)
 	{
 		RowItem = InArgs._Item;
@@ -64,19 +75,19 @@ public:
 	{
 		if (RowItem.IsValid())
 		{
-			if (ColumnName == ColumnId_AssetPath)
+			if (ColumnName == FMaterialStats::ColumnId_AssetPath)
 			{
 				return SNew(STextBlock).Text(FText::FromString(RowItem->AssetPath));
 			}
-			if (ColumnName == ColumnId_MaterialInstanceCount)
+			if (ColumnName == FMaterialStats::ColumnId_MaterialInstanceCount)
 			{
 				return SNew(STextBlock).Text(FText::AsNumber(RowItem->MaterialInstanceCount));
 			}
-			if (ColumnName == ColumnId_StaticPermutationCount)
+			if (ColumnName == FMaterialStats::ColumnId_StaticPermutationCount)
 			{
 				return SNew(STextBlock).Text(FText::AsNumber(RowItem->StaticPermutationCount));
 			}
-			if (ColumnName == ColumnId_ShaderCount)
+			if (ColumnName == FMaterialStats::ColumnId_ShaderCount)
 			{
 				return SNew(STextBlock).Text(FText::AsNumber(RowItem->ShaderCount));
 			}
@@ -86,7 +97,7 @@ public:
 };
 
 /** Delegate for notification of material data row selection. */
-DECLARE_DELEGATE_OneParam(FOnRowItemSelect, TSharedPtr<MaterialValidation::FMaterialStats>);
+DECLARE_DELEGATE_OneParam(FOnRowItemSelect, TSharedPtr<FMaterialStats>);
 
 /** A slate widget containing a table of UMaterial descriptions. */
 class SMaterialTable : public SCompoundWidget
@@ -119,12 +130,13 @@ public:
 	{
 		AllRows = InArgs._InitialRows;
 		FilteredRows = InArgs._InitialRows;
+		OnRowItemSelect = InArgs._OnRowItemSelect;
 
 		SearchText.Reset();
-		SortColumnId = NAME_None;
-		SortMode = EColumnSortMode::None;
-
-		OnRowItemSelect = InArgs._OnRowItemSelect;
+		
+		SortColumnId = FMaterialStats::ColumnId_StaticPermutationCount;
+		SortMode = EColumnSortMode::Descending;
+		UpdateSort();
 
 		ChildSlot
 		[
@@ -132,6 +144,7 @@ public:
 			+ SVerticalBox::Slot().AutoHeight().Padding(2)
 			[
 				SNew(SSearchBox)
+				.HintText(LOCTEXT("MaterialSearchHint", "Filter By Asset"))
 				.OnTextChanged(this, &SMaterialTable::OnSearchTextChanged)
 			]
 			+ SVerticalBox::Slot().FillHeight(1.0f).Padding(2)
@@ -141,10 +154,32 @@ public:
 				.HeaderRow
 				(
 					SNew(SHeaderRow)
-					+ SHeaderRow::Column(ColumnId_AssetPath).FillWidth(0.7).HAlignHeader(HAlign_Center).DefaultLabel(FText::FromName(ColumnId_AssetPath)).OnSort(this, &SMaterialTable::OnColumnSortModeChanged)
-					+ SHeaderRow::Column(ColumnId_MaterialInstanceCount).FillWidth(0.1).HAlignHeader(HAlign_Center).HAlignCell(HAlign_Center).DefaultLabel(FText::FromName(ColumnId_MaterialInstanceCount)).OnSort(this, &SMaterialTable::OnColumnSortModeChanged)
-					+ SHeaderRow::Column(ColumnId_StaticPermutationCount).FillWidth(0.1).HAlignHeader(HAlign_Center).HAlignCell(HAlign_Center).DefaultLabel(FText::FromName(ColumnId_StaticPermutationCount)).OnSort(this, &SMaterialTable::OnColumnSortModeChanged)
-					+ SHeaderRow::Column(ColumnId_ShaderCount).FillWidth(0.1).HAlignHeader(HAlign_Center).HAlignCell(HAlign_Center).DefaultLabel(FText::FromName(ColumnId_ShaderCount)).OnSort(this, &SMaterialTable::OnColumnSortModeChanged)
+					+ SHeaderRow::Column(FMaterialStats::ColumnId_AssetPath)
+						.HAlignHeader(HAlign_Center)
+						.HAlignCell(HAlign_Left)
+						.DefaultLabel(LOCTEXT("MaterialAssetPath", "Asset Path"))
+						.OnSort(this, &SMaterialTable::OnColumnSortModeChanged)
+						.ManualWidth(GAssetNameColumnWidth)
+						.ManualWidth_Lambda([]() { return GAssetNameColumnWidth; })
+						.OnWidthChanged_Lambda([](float NewWidth) { GAssetNameColumnWidth = NewWidth; })
+					+ SHeaderRow::Column(FMaterialStats::ColumnId_MaterialInstanceCount)
+						.HAlignHeader(HAlign_Center)
+						.HAlignCell(HAlign_Center)
+						.DefaultLabel(LOCTEXT("MaterialInstanceCount", "Instance Count"))
+						.OnSort(this, &SMaterialTable::OnColumnSortModeChanged)
+						.FixedWidth(150)
+					+ SHeaderRow::Column(FMaterialStats::ColumnId_StaticPermutationCount)
+						.HAlignHeader(HAlign_Center)
+						.HAlignCell(HAlign_Center)
+						.DefaultLabel(LOCTEXT("MaterialPermutationCount", "Permutation Count"))
+						.OnSort(this, &SMaterialTable::OnColumnSortModeChanged)
+						.FixedWidth(150)
+					+ SHeaderRow::Column(FMaterialStats::ColumnId_ShaderCount)
+						.HAlignHeader(HAlign_Center)
+						.HAlignCell(HAlign_Center)
+						.DefaultLabel(LOCTEXT("MaterialShaderCount", "Shader Count"))
+						.OnSort(this, &SMaterialTable::OnColumnSortModeChanged)
+						.FixedWidth(150)
 				)
 				.ListItemsSource(&FilteredRows)
 				.OnGenerateRow_Lambda([](TSharedPtr<FMaterialStats> InItem, const TSharedRef<STableViewBase>& InOwnerTable)
@@ -173,7 +208,7 @@ public:
 		if (SortColumnId != ColumnId)
 		{
 			// Initial sort mode is ascending for text, descending for numeric.
-			SortMode = ColumnId == ColumnId_AssetPath ? EColumnSortMode::Ascending : EColumnSortMode::Descending;
+			SortMode = ColumnId == FMaterialStats::ColumnId_AssetPath ? EColumnSortMode::Ascending : EColumnSortMode::Descending;
 			SortColumnId = ColumnId;
 		}
 		else
@@ -204,28 +239,28 @@ public:
 	{
 		const bool bSortModifier = SortMode == EColumnSortMode::Descending;
 
-		if (SortColumnId == ColumnId_AssetPath)
+		if (SortColumnId == FMaterialStats::ColumnId_AssetPath)
 		{
 			FilteredRows.Sort([bSortModifier](TSharedPtr<FMaterialStats> const& A, TSharedPtr<FMaterialStats> const& B)
 			{
 				return bSortModifier ? (A->AssetPath > B->AssetPath) : (A->AssetPath < B->AssetPath);
 			});
 		}
-		else if (SortColumnId == ColumnId_MaterialInstanceCount)
+		else if (SortColumnId == FMaterialStats::ColumnId_MaterialInstanceCount)
 		{
 			FilteredRows.Sort([bSortModifier](TSharedPtr<FMaterialStats> const& A, TSharedPtr<FMaterialStats> const& B)
 			{
 				return bSortModifier ? (A->MaterialInstanceCount > B->MaterialInstanceCount) : (A->MaterialInstanceCount < B->MaterialInstanceCount);
 			});
 		}
-		else if (SortColumnId == ColumnId_StaticPermutationCount)
+		else if (SortColumnId == FMaterialStats::ColumnId_StaticPermutationCount)
 		{
 			FilteredRows.Sort([bSortModifier](TSharedPtr<FMaterialStats> const& A, TSharedPtr<FMaterialStats> const& B)
 			{
 				return bSortModifier ? (A->StaticPermutationCount > B->StaticPermutationCount) : (A->StaticPermutationCount < B->StaticPermutationCount);
 			});
 		}
-		else if (SortColumnId == ColumnId_ShaderCount)
+		else if (SortColumnId == FMaterialStats::ColumnId_ShaderCount)
 		{
 			FilteredRows.Sort([bSortModifier](TSharedPtr<FMaterialStats> const& A, TSharedPtr<FMaterialStats> const& B)
 			{
@@ -244,46 +279,407 @@ public:
 	}
 };
 
+/** Instantiate a table widget for all materials in a group. */
+static TSharedPtr<SMaterialTable> CreateMaterialTable(UMaterialValidationGroup const* InGroup, FOnRowItemSelect OnRowItemSelect)
+{
+	TArray<TSharedPtr<FMaterialStats>> Stats;
+	if (InGroup != nullptr)
+	{
+		for (TPair<FString, FMaterialValidationDesc> const& It : InGroup->Materials)
+		{
+			Stats.Add(MakeShared<FMaterialStats>(It.Key, It.Value.MaterialInstances.Num(), It.Value.PermutationHashes.Num(), It.Value.NumShadersTotal));
+		}
+	}
+	return SNew(SMaterialTable).InitialRows(Stats).OnRowItemSelect(OnRowItemSelect);
+}
+
+/** The static switch and static component mask parameter FNames used by a single material. */
+struct FMaterialInstanceParameterNames
+{
+	// todo: Add fast FName->Index lookup because with large arrays, the linear search per table entry is very slow.
+	TArray<FName> StaticSwitches;
+	TArray<FName> ComponentMasks;
+};
+
+/** A description for a single column in the UMaterialInstance data table. */
+struct FMaterialInstanceColumnData
+{
+	/** FName Id used for column resolution. */
+	FName Id;
+	/** FText label used for display name. */
+	FText Label;
+	/** Column category used for filtering. */
+	EMaterialInstanceColumnType ColumnType;
+
+	FMaterialInstanceColumnData(FName const& InId, FText const& InLabel, EMaterialInstanceColumnType InColumnType)
+		: Id(InId)
+		, Label(InLabel)
+		, ColumnType(InColumnType)
+	{
+	}
+};
+
 /** A data table entry describing a UMaterialInstance. */
-struct FMaterialInstanceStats
+struct FMaterialInstanceRowData
 {
+	/** Path to the UMaterialInstance asset. */
 	FString AssetPath;
-	uint32 PermutationHash = 0;
+	/** Full resolved asset data description for the UMaterialInstance pulled from the material validation database. */
+	FMaterialValidationAssetData_MaterialInstance AssetData;
 
-	FMaterialInstanceStats(FString const& InPath, uint32 InPermutationHash)
+	FMaterialInstanceRowData(FString const& InPath, FMaterialValidationAssetData_MaterialInstance const& InAssetData)
 		: AssetPath(InPath)
-		, PermutationHash(InPermutationHash)
+		, AssetData(InAssetData)
+	{
+	}
+};
+
+/** 
+ * Enumeration of all of the UMaterialInstance base properties. 
+ * Aside from AssetPath and PermutationHash (kept in same enum for simplicity) these are the property overrides that trigger static permutations.
+ */
+enum class EMaterialInstanceStaticProperties : int32
+{
+	AssetPath,
+	PermutationHash,
+	TwoSided,
+	bIsThinSurface,
+	DitheredLODTransition,
+	bCastDynamicShadowAsMasked,
+	bOutputTranslucentVelocity,
+	bHasPixelAnimation,
+	bEnableTessellation,
+	BlendMode,
+	ShadingModel,
+	OpacityMaskClipValue,
+	Count
+};
+
+/** Helper getter functions for each of the EMaterialInstanceStaticProperties. */
+struct FMaterialInstanceStaticPropertyGetters
+{
+	static FString const& GetAssetPath(FMaterialInstanceRowData const& P) { return P.AssetPath; }
+	static uint32 GetPermutationHash(FMaterialInstanceRowData const& P) { return P.AssetData.PermutationHash; }
+	static uint8 GetTwoSided(FMaterialInstanceRowData const& P) { return P.AssetData.StaticProperties.TwoSided; }
+	static uint8 GetIsThinSurface(FMaterialInstanceRowData const& P) { return P.AssetData.StaticProperties.bIsThinSurface; }
+	static uint8 GetDitheredLODTransition(FMaterialInstanceRowData const& P) { return P.AssetData.StaticProperties.DitheredLODTransition; }
+	static uint8 GetCastDynamicShadowAsMasked(FMaterialInstanceRowData const& P) { return P.AssetData.StaticProperties.bCastDynamicShadowAsMasked; }
+	static uint8 GetOutputTranslucentVelocity(FMaterialInstanceRowData const& P) { return P.AssetData.StaticProperties.bOutputTranslucentVelocity; }
+	static uint8 GetHasPixelAnimation(FMaterialInstanceRowData const& P) { return P.AssetData.StaticProperties.bHasPixelAnimation; }
+	static uint8 GetEnableTessellation(FMaterialInstanceRowData const& P) { return P.AssetData.StaticProperties.bEnableTessellation; }
+	static TEnumAsByte<EBlendMode> GetBlendMode(FMaterialInstanceRowData const& P) { return P.AssetData.StaticProperties.BlendMode; }
+	static TEnumAsByte<EMaterialShadingModel> GetShadingModel(FMaterialInstanceRowData const& P) { return P.AssetData.StaticProperties.ShadingModel; }
+	static float GetOpacityMaskClipValue(FMaterialInstanceRowData const& P) { return P.AssetData.StaticProperties.OpacityMaskClipValue; }
+};
+
+/** Templatized sort functor for each of the EMaterialInstanceStaticProperties. */
+template <auto GetterPtr>
+static void SortByStaticProperty(TArray<TSharedPtr<FMaterialInstanceRowData>>& InArray, bool bInSortModifer)
+{
+	if (bInSortModifer)
+	{
+		Algo::StableSort(InArray, [](TSharedPtr<FMaterialInstanceRowData> const& A, TSharedPtr<FMaterialInstanceRowData> const& B) { return (*GetterPtr)(*A) > (*GetterPtr)(*B); });
+	}
+	else
+	{
+		Algo::StableSort(InArray, [](TSharedPtr<FMaterialInstanceRowData> const& A, TSharedPtr<FMaterialInstanceRowData> const& B) { return (*GetterPtr)(*A) < (*GetterPtr)(*B); });
+	}
+}
+
+using PropertySortFunction = void(*)(TArray<TSharedPtr<FMaterialInstanceRowData>>&, bool);
+
+/** Templatized table entry widget generator for each of the EMaterialInstanceStaticProperties. */
+template <auto GetterPtr, bool bFromString = false>
+static TSharedRef<SWidget> GenerateWidgetForStaticProperty(TSharedPtr<FMaterialInstanceRowData> const& RowItem)
+{
+	if constexpr(bFromString)
 	{
+		return SNew(STextBlock).Text(FText::FromString((*GetterPtr)(*RowItem)));
 	}
+	else
+	{
+		return SNew(STextBlock).Text(FText::AsNumber((*GetterPtr)(*RowItem)));
+	}
+}
+
+using PropertyGenerateWidgetFunction = TSharedRef<SWidget>(*)(TSharedPtr<FMaterialInstanceRowData> const&);
+
+/** Struct containing information needed to generate table entries for EMaterialInstanceStaticProperties. */
+struct FMaterialInstanceStaticPropertyInfo
+{
+	EMaterialInstanceColumnType Type;
+	FName ColumnId;
+	FString PropertyName;
+	PropertySortFunction SortFunction;
+	PropertyGenerateWidgetFunction GenerateWidgetFunction;
 };
 
+/** Array for table generation data for EMaterialInstanceStaticProperties. */
+static const FMaterialInstanceStaticPropertyInfo StaticPropertyInfos[(int32)EMaterialInstanceStaticProperties::Count] = {
+	{EMaterialInstanceColumnType::Default, "CID_AssetPath", "Asset Path", &SortByStaticProperty<&FMaterialInstanceStaticPropertyGetters::GetAssetPath>, &GenerateWidgetForStaticProperty<&FMaterialInstanceStaticPropertyGetters::GetAssetPath, true> },
+	{EMaterialInstanceColumnType::Default, "CID_PermutationHash", "Permutation Hash", &SortByStaticProperty<&FMaterialInstanceStaticPropertyGetters::GetPermutationHash>, &GenerateWidgetForStaticProperty<&FMaterialInstanceStaticPropertyGetters::GetPermutationHash> },
+	{EMaterialInstanceColumnType::StaticProperty, "CID_TwoSided", "Two Sided", &SortByStaticProperty<&FMaterialInstanceStaticPropertyGetters::GetTwoSided>, &GenerateWidgetForStaticProperty<&FMaterialInstanceStaticPropertyGetters::GetTwoSided> },
+	{EMaterialInstanceColumnType::StaticProperty, "CID_IsThinSurface", "Thin Surface", &SortByStaticProperty<&FMaterialInstanceStaticPropertyGetters::GetIsThinSurface>, &GenerateWidgetForStaticProperty<&FMaterialInstanceStaticPropertyGetters::GetIsThinSurface> },
+	{EMaterialInstanceColumnType::StaticProperty, "CID_DitheredLODTransition", "Dithered LOD Transition", &SortByStaticProperty<&FMaterialInstanceStaticPropertyGetters::GetDitheredLODTransition>, &GenerateWidgetForStaticProperty<&FMaterialInstanceStaticPropertyGetters::GetDitheredLODTransition> },
+	{EMaterialInstanceColumnType::StaticProperty, "CID_CastDynamicShadowAsMasked", "Dynamic Shadow As Masked", &SortByStaticProperty<&FMaterialInstanceStaticPropertyGetters::GetCastDynamicShadowAsMasked>, &GenerateWidgetForStaticProperty<&FMaterialInstanceStaticPropertyGetters::GetCastDynamicShadowAsMasked> },
+	{EMaterialInstanceColumnType::StaticProperty, "CID_OutputTranslucentVelocity", "Output Translucent Velocity", &SortByStaticProperty<&FMaterialInstanceStaticPropertyGetters::GetOutputTranslucentVelocity>, &GenerateWidgetForStaticProperty<&FMaterialInstanceStaticPropertyGetters::GetOutputTranslucentVelocity> },
+	{EMaterialInstanceColumnType::StaticProperty, "CID_HasPixelAnimation", "Has Pixel Animation", &SortByStaticProperty<&FMaterialInstanceStaticPropertyGetters::GetHasPixelAnimation>, &GenerateWidgetForStaticProperty<&FMaterialInstanceStaticPropertyGetters::GetHasPixelAnimation> },
+	{EMaterialInstanceColumnType::StaticProperty, "CID_EnableTessellation", "Tessellation", &SortByStaticProperty<&FMaterialInstanceStaticPropertyGetters::GetEnableTessellation>, &GenerateWidgetForStaticProperty<&FMaterialInstanceStaticPropertyGetters::GetEnableTessellation> },
+	{EMaterialInstanceColumnType::StaticProperty, "CID_BlendMode", "Blend Mode", &SortByStaticProperty<&FMaterialInstanceStaticPropertyGetters::GetBlendMode>, &GenerateWidgetForStaticProperty<&FMaterialInstanceStaticPropertyGetters::GetBlendMode> },
+	{EMaterialInstanceColumnType::StaticProperty, "CID_ShadingModel", "Shading Model", &SortByStaticProperty<&FMaterialInstanceStaticPropertyGetters::GetShadingModel>, &GenerateWidgetForStaticProperty<&FMaterialInstanceStaticPropertyGetters::GetShadingModel> },
+	{EMaterialInstanceColumnType::StaticProperty, "CID_OpacityMaskClipValue", "Opacity Mask Clip Value", &SortByStaticProperty<&FMaterialInstanceStaticPropertyGetters::GetOpacityMaskClipValue>, &GenerateWidgetForStaticProperty<&FMaterialInstanceStaticPropertyGetters::GetOpacityMaskClipValue> },
+};
+
+/** Get column FName for a static property enumeration. */
+static FName GetColumnIdFromStaticProperty(EMaterialInstanceStaticProperties InProperty)
+{
+	if ((int32)InProperty >= 0 && (int32)InProperty < (int32)EMaterialInstanceStaticProperties::Count)
+	{
+		return StaticPropertyInfos[(int32)InProperty].ColumnId;
+	}
+	return NAME_None;
+}
+
+/** Resolve column FName to static property enumeration. */
+static EMaterialInstanceStaticProperties GetStaticPropertyFromColumnId(FName InColumnId)
+{
+	for (int32 Index = 0; Index < (int32)EMaterialInstanceStaticProperties::Count; ++Index)
+	{
+		if (StaticPropertyInfos[Index].ColumnId == InColumnId)
+		{
+			return (EMaterialInstanceStaticProperties)Index;
+		}
+	}
+	return EMaterialInstanceStaticProperties::Count;
+}
+
+/** Sort row data based on the given static property enumeration. */
+static void SortByStaticProperty(TArray<TSharedPtr<FMaterialInstanceRowData>>& InArray, EMaterialInstanceStaticProperties InKey, bool bInSortModifer)
+{
+	if (InKey != EMaterialInstanceStaticProperties::Count)
+	{
+		StaticPropertyInfos[(uint8)InKey].SortFunction(InArray, bInSortModifer);
+	}
+}
+
+/** Generate row entry widget for the given static property enumeration. */
+static TSharedRef<SWidget> GenerateWidgetForStaticProperty(TSharedPtr<FMaterialInstanceRowData> const& RowItem, EMaterialInstanceStaticProperties InKey)
+{
+	if (InKey != EMaterialInstanceStaticProperties::Count)
+	{
+		return StaticPropertyInfos[(uint8)InKey].GenerateWidgetFunction(RowItem);
+	}
+
+	return SNew(STextBlock).Text(FText::FromString("?"));
+}
+
+/** Static array of column FNames for all material usages. */
+static const FName UsageColumnIds[MATUSAGE_MAX] = {
+	"CID_UsageSkeletalMesh",
+	"CID_UsageParticleSprites",
+	"CID_UsageBeamTrails",
+	"CID_UsageMeshParticles",
+	"CID_UsageStaticLighting",
+	"CID_UsageMorphTargets",
+	"CID_UsageSplineMesh",
+	"CID_UsageInstancedStaticMeshes",
+	"CID_UsageGeometryCollections",
+	"CID_UsageClothing",
+	"CID_UsageNiagaraSprites",
+	"CID_UsageNiagaraRibbons",
+	"CID_UsageNiagaraMeshParticles",
+	"CID_UsageGeometryCache",
+	"CID_UsageWater",
+	"CID_UsageHairStrands",
+	"CID_UsageLidarPointCloud",
+	"CID_UsageVirtualHeightfieldMesh",
+	"CID_UsageNanite",
+	"CID_UsageVoxels",
+	"CID_UsageVolumetricCloud",
+	"CID_UsageHeterogeneousVolumes",
+	"CID_UsageStaticMesh",
+	"CID_UsageEditorCompositing",
+	"CID_UsageNeuralNetworks",
+	"CID_UsageMeshDeformer",
+};
+
+/** Resolve column FName to material usage. */
+static EMaterialUsage GetMaterialUsageFromColumnId(FName InColumnId)
+{
+	for (int32 Index = 0; Index < MATUSAGE_MAX; ++Index)
+	{
+		if (UsageColumnIds[Index] == InColumnId)
+		{
+			return (EMaterialUsage)Index;
+		}
+	}
+	return MATUSAGE_MAX;
+}
+
+/** Sort row data based on the given material usage. */
+static void SortByUsage(TArray<TSharedPtr<FMaterialInstanceRowData>>& InArray, EMaterialUsage InUsage, bool bInSortModifer)
+{
+	if (InUsage < MATUSAGE_MAX)
+	{
+		if (bInSortModifer)
+		{
+			Algo::StableSort(InArray, [InUsage](TSharedPtr<FMaterialInstanceRowData> const& A, TSharedPtr<FMaterialInstanceRowData> const& B)
+			{
+				return (A->AssetData.StaticProperties.UsageFlags & (1u << InUsage)) > (B->AssetData.StaticProperties.UsageFlags & (1u << InUsage));
+			});
+		}
+		else
+		{
+			Algo::StableSort(InArray, [InUsage](TSharedPtr<FMaterialInstanceRowData> const& A, TSharedPtr<FMaterialInstanceRowData> const& B)
+			{
+				return (A->AssetData.StaticProperties.UsageFlags & (1u << InUsage)) < (B->AssetData.StaticProperties.UsageFlags & (1u << InUsage));
+			});
+		}
+	}
+}
+
+/** Resolve column FName to index of static switch parameter. Returns INDEX_NONE if not found. */
+static int32 GetStaticSwitchIndexFromColumnId(FName InColumnId, FMaterialInstanceParameterNames const& InParameterNames)
+{
+	for (int32 Index = 0; Index < InParameterNames.StaticSwitches.Num(); ++Index)
+	{
+		if (InParameterNames.StaticSwitches[Index] == InColumnId)
+		{
+			return Index;
+		}
+	}
+	return INDEX_NONE;
+}
+
+/** Sort row data based on the given static switch parameter index. */
+static void SortByStaticSwitch(TArray<TSharedPtr<FMaterialInstanceRowData>>& InArray, int32 InIndex, bool bInSortModifer)
+{
+	if (InIndex < 0)
+	{
+		return;
+	}
+
+	const uint32 BitArrayIndex = (uint32)InIndex / 32;
+	const uint32 BitIndex = (uint32)InIndex % 32;
+	if (bInSortModifer)
+	{
+		Algo::StableSort(InArray, [BitArrayIndex, BitIndex](TSharedPtr<FMaterialInstanceRowData> const& A, TSharedPtr<FMaterialInstanceRowData> const& B)
+		{
+			return (A->AssetData.StaticSwitchOverrideValues[BitArrayIndex] & (1u << BitIndex)) > (B->AssetData.StaticSwitchOverrideValues[BitArrayIndex] & (1u << BitIndex));
+		});
+	}
+	else
+	{
+		Algo::StableSort(InArray, [BitArrayIndex, BitIndex](TSharedPtr<FMaterialInstanceRowData> const& A, TSharedPtr<FMaterialInstanceRowData> const& B)
+		{
+			return (A->AssetData.StaticSwitchOverrideValues[BitArrayIndex] & (1u << BitIndex)) < (B->AssetData.StaticSwitchOverrideValues[BitArrayIndex] & (1u << BitIndex));
+		});
+	}
+}
+
+/** Resolve column FName to index of static component mask parameter. Returns INDEX_NONE if not found. */
+static int32 GetComponentMaskIndexFromColumnId(FName InColumnId, FMaterialInstanceParameterNames const& InParameterNames)
+{
+	for (int32 Index = 0; Index < InParameterNames.ComponentMasks.Num(); ++Index)
+	{
+		if (InParameterNames.ComponentMasks[Index] == InColumnId)
+		{
+			return Index;
+		}
+	}
+	return INDEX_NONE;
+}
+
+/** Sort row data based on the given static componenent mask parameter index. */
+static void SortByComponentMask(TArray<TSharedPtr<FMaterialInstanceRowData>>& InArray, int32 InIndex, bool bInSortModifer)
+{
+	if (InIndex < 0)
+	{
+		return;
+	}
+
+	const uint32 BitArrayIndex = InIndex / 8;
+	const uint32 BitIndex = (InIndex * 4) % 32;
+	if (bInSortModifer)
+	{
+		Algo::StableSort(InArray, [BitArrayIndex, BitIndex](TSharedPtr<FMaterialInstanceRowData> const& A, TSharedPtr<FMaterialInstanceRowData> const& B)
+		{
+			return ((A->AssetData.ComponentMaskOverrideValues[BitArrayIndex] >> BitIndex) & 0xf) > ((B->AssetData.ComponentMaskOverrideValues[BitArrayIndex] >> BitIndex) & 0xf);
+		});
+	}
+	else
+	{
+		Algo::StableSort(InArray, [BitArrayIndex, BitIndex](TSharedPtr<FMaterialInstanceRowData> const& A, TSharedPtr<FMaterialInstanceRowData> const& B)
+		{
+			return ((A->AssetData.ComponentMaskOverrideValues[BitArrayIndex] >> BitIndex) & 0xf) < ((B->AssetData.ComponentMaskOverrideValues[BitArrayIndex] >> BitIndex) & 0xf);
+		});
+	}
+}
+
 /** A slate container of a data table entry describing a UMaterialInstance. */
-class SMaterialInstanceTableRow : public SMultiColumnTableRow<TSharedPtr<FMaterialInstanceStats>>
+class SMaterialInstanceTableRow : public SMultiColumnTableRow<TSharedPtr<FMaterialInstanceRowData>>
 {
-	TSharedPtr<FMaterialInstanceStats> RowItem;
+	/** Shared pointer to container of resoved static parameter names owned by the table. */
+	TSharedPtr<FMaterialInstanceParameterNames> ParameterNames;
+	/** The material instance stats for the row. */
+	TSharedPtr<FMaterialInstanceRowData> RowItem;
 
 public:
 	SLATE_BEGIN_ARGS(SMaterialInstanceTableRow) {}
-		SLATE_ARGUMENT(TSharedPtr<FMaterialInstanceStats>, Item)
+		SLATE_ARGUMENT(TSharedPtr<FMaterialInstanceParameterNames>, ParameterNames)
+		SLATE_ARGUMENT(TSharedPtr<FMaterialInstanceRowData>, Item)
 	SLATE_END_ARGS()
 
+	/** Builds the row. */
 	void Construct(const FArguments& InArgs, const TSharedRef<STableViewBase>& InOwnerTable)
 	{
+		ParameterNames = InArgs._ParameterNames;
 		RowItem = InArgs._Item;
-		SMultiColumnTableRow<TSharedPtr<FMaterialInstanceStats>>::Construct(FSuperRowType::FArguments(), InOwnerTable);
+		SMultiColumnTableRow<TSharedPtr<FMaterialInstanceRowData>>::Construct(FSuperRowType::FArguments(), InOwnerTable);
 	}
 
-	virtual TSharedRef<SWidget> GenerateWidgetForColumn(const FName& ColumnName) override
+	virtual TSharedRef<SWidget> GenerateWidgetForColumn(const FName& ColumnId) override
 	{
 		if (RowItem.IsValid())
 		{
-			if (ColumnName == ColumnId_AssetPath)
+			// Try to resolve column as static property.
+			const EMaterialInstanceStaticProperties Property = GetStaticPropertyFromColumnId(ColumnId);
+			if (Property != EMaterialInstanceStaticProperties::Count)
 			{
-				return SNew(STextBlock).Text(FText::FromString(RowItem->AssetPath));
+				return GenerateWidgetForStaticProperty(RowItem, Property);
+			}
+
+			// Try to resolve column as material usage.
+			const EMaterialUsage Usage = GetMaterialUsageFromColumnId(ColumnId);
+			if (Usage != MATUSAGE_MAX)
+			{
+				const bool bValue = (RowItem->AssetData.StaticProperties.UsageFlags & (1u << Usage)) != 0;
+				return SNew(STextBlock).Text(FText::AsNumber(bValue ? 1 : 0));
+			}
+
+			// Try to resolve column as static switch.
+			const int32 StaticSwitchIndex = GetStaticSwitchIndexFromColumnId(ColumnId, *ParameterNames);
+			if (StaticSwitchIndex != INDEX_NONE)
+			{
+				const uint32 BitArrayIndex = StaticSwitchIndex / 32;
+				const uint32 BitIndex = StaticSwitchIndex % 32;
+				const bool bValue = (RowItem->AssetData.StaticSwitchOverrideValues[BitArrayIndex] & (1 << BitIndex)) != 0;
+				return SNew(STextBlock).Text(FText::AsNumber(bValue ? 1 : 0));
 			}
-			if (ColumnName == ColumnId_PermutationHash)
+
+			// Try to resolve column as static component mask.
+			const int32 ComponentMaskIndex = GetComponentMaskIndexFromColumnId(ColumnId, *ParameterNames);
+			if (ComponentMaskIndex != INDEX_NONE)
 			{
-				return SNew(STextBlock).Text(FText::AsNumber(RowItem->PermutationHash));
+				const uint32 BitArrayIndex = ComponentMaskIndex / 8;
+				const uint32 BitIndex = (ComponentMaskIndex * 4) % 32;
+				const uint32 Value = RowItem->AssetData.ComponentMaskOverrideValues[BitArrayIndex] >> BitIndex;
+				TStringBuilder<8> RGBA;
+				RGBA.AppendChar(Value & 1 ? 'R' : '_');
+				RGBA.AppendChar(Value & 2 ? 'G' : '_');
+				RGBA.AppendChar(Value & 4 ? 'B' : '_');
+				RGBA.AppendChar(Value & 8 ? 'A' : '_');
+				return SNew(STextBlock).Text(FText::FromString(*RGBA));
 			}
 		}
 		return SNew(STextBlock).Text(FText::FromString("?"));
@@ -293,15 +689,32 @@ public:
 /** A slate widget containing a table of UMaterialInstance descriptions. */
 class SMaterialInstanceTable : public SCompoundWidget
 {
+	/** All the available columns. */
+	TArray<TSharedPtr<FMaterialInstanceColumnData>> AllColumns;
+	/** An ordered and filtered subset of AllColumns. */
+	TArray<TSharedPtr<FMaterialInstanceColumnData>> FilteredColumns;
+
 	/** All the available rows. */
-	TArray<TSharedPtr<FMaterialInstanceStats>> AllRows;
+	TArray<TSharedPtr<FMaterialInstanceRowData>> AllRows;
 	/** An ordered and filtered subset of AllRows. */
-	TArray<TSharedPtr<FMaterialInstanceStats>> FilteredRows;
+	TArray<TSharedPtr<FMaterialInstanceRowData>> FilteredRows;
+	
+	/** Resolved static switch and component mask parameter names. */
+	TSharedPtr<FMaterialInstanceParameterNames> ParameterNames;
+
+	/** Table header row. */
+	TSharedPtr<SHeaderRow> HeaderRow;
 	/** A slate view of the table. */
-	TSharedPtr<SListView<TSharedPtr<FMaterialInstanceStats>>> ListView;
+	TSharedPtr<SListView<TSharedPtr<FMaterialInstanceRowData>>> ListView;
+
+	/** Filter for type of columns to show. */
+	EMaterialInstanceColumnType SelectedColumnType;
+
+	/** Search text used for filtering the rows to show. */
+	FString RowSearchText;
+	/** Search text used for filtering the columns to show. */
+	FString ColumnSearchText;
 
-	/** Search text used for filtering the rows. */
-	FString SearchText;
 	/** Selected column name for sorting the table. */
 	FName SortColumnId = NAME_None;
 	/** Selected mode for sorting the table. */
@@ -309,170 +722,325 @@ class SMaterialInstanceTable : public SCompoundWidget
 
 public:
 	SLATE_BEGIN_ARGS(SMaterialInstanceTable) {}
-		SLATE_ARGUMENT(TArray<TSharedPtr<FMaterialInstanceStats>>, InitialRows)
+		SLATE_ARGUMENT(TArray<TSharedPtr<FMaterialInstanceColumnData>>, InitialColumns)
+		SLATE_ARGUMENT(TArray<TSharedPtr<FMaterialInstanceRowData>>, InitialRows)
+		SLATE_ARGUMENT(TSharedPtr<FMaterialInstanceParameterNames>, InitialParameterNames)
 	SLATE_END_ARGS()
 
 	/** Builds the slate table view. */
 	void Construct(const FArguments& InArgs)
 	{
+		AllColumns = InArgs._InitialColumns;
+		FilteredColumns = InArgs._InitialColumns;
 		AllRows = InArgs._InitialRows;
 		FilteredRows = InArgs._InitialRows;
 
-		SearchText.Reset();
-		SortColumnId = NAME_None;
-		SortMode = EColumnSortMode::None;
+		ParameterNames = InArgs._InitialParameterNames;
+
+		SelectedColumnType = EMaterialInstanceColumnType::StaticProperty;
+
+		RowSearchText.Reset();
+		ColumnSearchText.Reset();
+
+		SortColumnId = GetColumnIdFromStaticProperty(EMaterialInstanceStaticProperties::PermutationHash);
+		SortMode = EColumnSortMode::Descending;
+
+		HeaderRow = SNew(SHeaderRow);
+		UpdateHeaderRow();
 
 		ChildSlot
 		[
 			SNew(SVerticalBox)
-			+ SVerticalBox::Slot().AutoHeight().Padding(2)
+			+SVerticalBox::Slot().AutoHeight()
 			[
-				SNew(SSearchBox)
-				.OnTextChanged(this, &SMaterialInstanceTable::OnSearchTextChanged)
+				SNew(SHorizontalBox)
+				+ SHorizontalBox::Slot().FillWidth(0.2).Padding(2)
+				[
+					SNew(SComboButton)
+					.OnGetMenuContent(this, &SMaterialInstanceTable::OnGetColumnTypeContent)
+					.ButtonContent()
+					[
+						SNew(STextBlock)
+						.Text(this, &SMaterialInstanceTable::GetSelectedColumnTypeText)
+					]
+				]
+				+ SHorizontalBox::Slot().FillWidth(0.4).Padding(2)
+				[
+					SNew(SSearchBox)
+					.HintText(LOCTEXT("MaterialInstanceSearchHint", "Filter By Asset"))
+					.OnTextChanged(this, &SMaterialInstanceTable::OnRowSearchTextChanged)
+				]
+				+SHorizontalBox::Slot().FillWidth(0.4).Padding(2)
+				[
+					SNew(SSearchBox)
+					.HintText(LOCTEXT("ColumnSearchHint", "Filter Columns"))
+					.OnTextChanged(this, &SMaterialInstanceTable::OnColumnSearchTextChanged)
+				]
 			]
-			+ SVerticalBox::Slot().FillHeight(1.0f).Padding(2)
-			.FillHeight(1.f)
+			+SVerticalBox::Slot().HAlign(HAlign_Fill).FillHeight(1.0f)
 			[
-				SAssignNew(ListView, SListView<TSharedPtr<FMaterialInstanceStats>>)
-				.HeaderRow
-				(
-					SNew(SHeaderRow)
-					+ SHeaderRow::Column(ColumnId_AssetPath).FillWidth(0.7).HAlignHeader(HAlign_Center).DefaultLabel(FText::FromName(ColumnId_AssetPath)).OnSort(this, &SMaterialInstanceTable::OnColumnSortModeChanged)
-					+ SHeaderRow::Column(ColumnId_PermutationHash).FillWidth(0.3).HAlignHeader(HAlign_Center).HAlignCell(HAlign_Center).DefaultLabel(FText::FromName(ColumnId_PermutationHash)).OnSort(this, &SMaterialInstanceTable::OnColumnSortModeChanged)
-				)
-				.ListItemsSource(&FilteredRows)
-				.OnGenerateRow_Lambda([](TSharedPtr<FMaterialInstanceStats> InItem, const TSharedRef<STableViewBase>& InOwnerTable)
-				{
-					return SNew(SMaterialInstanceTableRow, InOwnerTable).Item(InItem);
-				})
+				SNew(SScrollBox)
+				.Orientation(Orient_Horizontal)
+				+SScrollBox::Slot()
+				[
+					SAssignNew(ListView, SListView<TSharedPtr<FMaterialInstanceRowData>>)
+					.HeaderRow(HeaderRow.ToSharedRef())
+					.ListItemsSource(&FilteredRows)
+					.OnGenerateRow_Lambda([this](TSharedPtr<FMaterialInstanceRowData> InItem, const TSharedRef<STableViewBase>& InOwnerTable)
+					{
+						return SNew(SMaterialInstanceTableRow, InOwnerTable).Item(InItem).ParameterNames(ParameterNames);
+					})
+				]
 			]
 		];
 	}
 
 	/** Set the table content. */
-	void SetRows(TArray<TSharedPtr<FMaterialInstanceStats>> const& InRows)
+	void SetContent(
+		TArray<TSharedPtr<FMaterialInstanceRowData>> const& InRows,
+		TArray<TSharedPtr<FMaterialInstanceColumnData>> const& InColumns,
+		TSharedPtr<FMaterialInstanceParameterNames> InParameterNames)
 	{
 		AllRows = InRows;
+		AllColumns = InColumns;
+		ParameterNames = InParameterNames;
 
-		UpdateFilter();
+		UpdateRowFilter();
+		UpdateColumnFilter();
+		UpdateHeaderRow();
 		UpdateSort();
 		UpdateListView();
 	}
 
-	/** Called when modifying the search text. */
-	void OnSearchTextChanged(const FText& InText)
+protected:
+	/** Called to fill EMaterialInstanceColumnType filter combo box. */
+	TSharedRef<SWidget> OnGetColumnTypeContent()
 	{
-		SearchText = InText.ToString();
-
-		UpdateFilter();
-		UpdateSort();
-		UpdateListView();
+		const UEnum* Enum = StaticEnum<EMaterialInstanceColumnType>();
+		FMenuBuilder MenuBuilder(true, NULL);
+		for (int32 Index = 1; Index < (int32)EMaterialInstanceColumnType::Count; ++Index)
+		{
+			FUIAction ItemAction(FExecuteAction::CreateSP(this, &SMaterialInstanceTable::OnSelectColumnType, Index));
+			MenuBuilder.AddMenuEntry(Enum->GetDisplayNameTextByIndex(Index), TAttribute<FText>(), FSlateIcon(), ItemAction);
+		}
+		return MenuBuilder.MakeWidget();
 	}
 
-	/** Called when clicking on a column header to sort. */
-	void OnColumnSortModeChanged(const EColumnSortPriority::Type SortPriority, const FName& ColumnId, const EColumnSortMode::Type InSortMode)
+	/** Called to get current EMaterialInstanceColumnType filter text entry. */
+	FText GetSelectedColumnTypeText() const
 	{
-		if (SortColumnId != ColumnId)
+		const UEnum* Enum = StaticEnum<EMaterialInstanceColumnType>();
+		return Enum->GetDisplayNameTextByIndex((int32)SelectedColumnType);
+	}
+	
+	/** Update the view. Must be called after any work that updates filtered/sorted content. */
+	void UpdateListView()
+	{
+		if (ListView.IsValid())
 		{
-			// Initial sort mode is ascending for text, descending for numeric.
-			SortMode = ColumnId == ColumnId_AssetPath ? EColumnSortMode::Ascending : EColumnSortMode::Descending;
-			SortColumnId = ColumnId;
+			ListView->RequestListRefresh();
 		}
-		else
+	}
+
+	/** Apply header row. Must be called after any work that updates filtered columns. */
+	void UpdateHeaderRow()
+	{
+		HeaderRow->ClearColumns();
+		for (TSharedPtr<FMaterialInstanceColumnData> const& Column : FilteredColumns)
 		{
-			SortMode = SortMode == EColumnSortMode::Descending ? EColumnSortMode::Ascending : EColumnSortMode::Descending;
+			SHeaderRow::FColumn::FArguments Args = SHeaderRow::Column(Column->Id);
+			Args.DefaultLabel(Column->Label).OnSort(this, &SMaterialInstanceTable::OnColumnSortModeChanged);
+			if (Column->Id == GetColumnIdFromStaticProperty(EMaterialInstanceStaticProperties::AssetPath))
+			{
+				Args
+				.HAlignHeader(HAlign_Center)
+				.HAlignCell(HAlign_Left)
+				.ManualWidth(GAssetNameColumnWidth)
+				.ManualWidth_Lambda([]() { return GAssetNameColumnWidth; })
+				.OnWidthChanged_Lambda([](float NewWidth) { GAssetNameColumnWidth = NewWidth; });
+			}
+			else
+			{
+				Args
+				.HAlignHeader(HAlign_Center)
+				.HAlignCell(HAlign_Center)
+				.FixedWidth(120);
+			}
+			HeaderRow->AddColumn(Args);
 		}
+	}
+
+	/** Called on selection from EMaterialInstanceColumnType filter combo box. */
+	void OnSelectColumnType(int32 Index)
+	{
+		SelectedColumnType = (EMaterialInstanceColumnType)Index;
+
+		UpdateColumnFilter();
+		UpdateHeaderRow();
+		UpdateListView();
+	}
+
+	/** Called when modifying the row search text. */
+	void OnRowSearchTextChanged(const FText& InText)
+	{
+		RowSearchText = InText.ToString();
 
+		UpdateRowFilter();
 		UpdateSort();
 		UpdateListView();
 	}
 
 	/**  Apply filter to AllRows to recreate FilteredRows. */
-	void UpdateFilter()
+	void UpdateRowFilter()
 	{
 		FilteredRows.Reset();
 
-		for (const TSharedPtr<FMaterialInstanceStats>& Row : AllRows)
+		for (const TSharedPtr<FMaterialInstanceRowData>& Row : AllRows)
 		{
-			if (SearchText.IsEmpty() || Row->AssetPath.Contains(SearchText, ESearchCase::IgnoreCase))
+			if (RowSearchText.IsEmpty() || Row->AssetPath.Contains(RowSearchText, ESearchCase::IgnoreCase))
 			{
 				FilteredRows.Add(Row);
 			}
 		}
 	}
 
-	/**  Apply sort to FilteredRows. */
-	void UpdateSort()
+	/** Called when modifying the column search text. */
+	void OnColumnSearchTextChanged(const FText& InText)
 	{
-		const bool bSortModifier = SortMode == EColumnSortMode::Descending;
+		ColumnSearchText = InText.ToString();
 
-		if (SortColumnId == ColumnId_AssetPath)
-		{
-			FilteredRows.Sort([bSortModifier](TSharedPtr<FMaterialInstanceStats> const& A, TSharedPtr<FMaterialInstanceStats> const& B)
-			{
-				return bSortModifier ? (A->AssetPath > B->AssetPath) : (A->AssetPath < B->AssetPath);
-			});
-		}
-		else if (SortColumnId == ColumnId_PermutationHash)
+		UpdateColumnFilter();
+		UpdateHeaderRow();
+		UpdateListView();
+	}
+
+	/**  Apply filter to AllColumns to recreate FilteredColumns. */
+	void UpdateColumnFilter()
+	{
+		FilteredColumns.Reset();
+
+		for (const TSharedPtr<FMaterialInstanceColumnData>& Column : AllColumns)
 		{
-			FilteredRows.Sort([bSortModifier](TSharedPtr<FMaterialInstanceStats> const& A, TSharedPtr<FMaterialInstanceStats> const& B)
+			if (Column->ColumnType == EMaterialInstanceColumnType::Default || 
+				(Column->ColumnType == SelectedColumnType && (ColumnSearchText.IsEmpty() || Column->Label.ToString().Contains(ColumnSearchText, ESearchCase::IgnoreCase))))
 			{
-				return bSortModifier ? (A->PermutationHash > B->PermutationHash) : (A->PermutationHash < B->PermutationHash);
-			});
+				FilteredColumns.Add(Column);
+			}
 		}
 	}
 
-	/** Update the view. Must be called after any call to UpdateFilter() or UpdateSort(). */
-	void UpdateListView()
+	/** Called when clicking on a column header to sort. */
+	void OnColumnSortModeChanged(const EColumnSortPriority::Type SortPriority, const FName& ColumnId, const EColumnSortMode::Type InSortMode)
 	{
-		if (ListView.IsValid())
+		if (SortColumnId != ColumnId)
 		{
-			ListView->RequestListRefresh();
+			// Initial sort mode is ascending for text, descending for numeric.
+			SortMode = ColumnId == GetColumnIdFromStaticProperty(EMaterialInstanceStaticProperties::AssetPath) ? EColumnSortMode::Ascending : EColumnSortMode::Descending;
+			SortColumnId = ColumnId;
 		}
+		else
+		{
+			SortMode = SortMode == EColumnSortMode::Descending ? EColumnSortMode::Ascending : EColumnSortMode::Descending;
+		}
+
+		UpdateSort();
+		UpdateListView();
 	}
-};
 
-static TSharedPtr<SMaterialTable> CreateMaterialTable(UMaterialValidationGroup const* InGroup, FOnRowItemSelect OnRowItemSelect)
-{
-	TArray<TSharedPtr<FMaterialStats>> Stats;
-	if (InGroup != nullptr)
+	/**  Apply sort to FilteredRows. */
+	void UpdateSort()
 	{
-		for (TPair<FString, FMaterialValidationDesc> const& It : InGroup->Materials)
+		TSharedPtr<FMaterialInstanceColumnData>* SortColumnData = FilteredColumns.FindByPredicate([ColumnId = SortColumnId](TSharedPtr<FMaterialInstanceColumnData> const& ColumnData)
 		{
-			Stats.Add(MakeShared<FMaterialStats>(It.Key, It.Value.MaterialInstances.Num(), It.Value.PermutationHashes.Num(), It.Value.NumShadersTotal));
+			return ColumnData->Id == ColumnId; 
+		});
+		if (SortColumnData == nullptr)
+		{
+			// This can happen if we are filtering out the sort column.
+			// Current behavior is no sorting until a column to sort is next selected.
+			return;
 		}
-	}
-	return SNew(SMaterialTable).InitialRows(Stats).OnRowItemSelect(OnRowItemSelect);
-}
 
-static TSharedPtr<SMaterialInstanceTable> CreateMaterialInstanceTable(FMaterialValidationDesc const* InMaterialDesc)
-{
-	TArray<TSharedPtr<FMaterialInstanceStats>> Stats;
-	if (InMaterialDesc != nullptr)
-	{
-		for (TPair<FString, int32> const& It : InMaterialDesc->MaterialInstances)
+		// Call specific sorting function for the column type.
+		const bool bSortModifier = SortMode == EColumnSortMode::Descending;
+		const EMaterialInstanceColumnType ColumnType = (*SortColumnData)->ColumnType;
+		switch (ColumnType)
 		{
-			const int32 AssetDataIndex = It.Value;
-			const uint32 PermutationHash = InMaterialDesc->MaterialInstanceAssetDatas.IsValidIndex(AssetDataIndex) ? InMaterialDesc->MaterialInstanceAssetDatas[AssetDataIndex].PermutationHash : 0;
-			Stats.Add(MakeShared<FMaterialInstanceStats>(It.Key, PermutationHash));
+			case EMaterialInstanceColumnType::Default:
+			case EMaterialInstanceColumnType::StaticProperty:
+				SortByStaticProperty(FilteredRows, GetStaticPropertyFromColumnId(SortColumnId), bSortModifier);
+				break;
+			case EMaterialInstanceColumnType::UsageFlag:
+				SortByUsage(FilteredRows, GetMaterialUsageFromColumnId(SortColumnId), bSortModifier);
+				break;
+			case EMaterialInstanceColumnType::StaticSwitch:
+				SortByStaticSwitch(FilteredRows, GetStaticSwitchIndexFromColumnId(SortColumnId, *ParameterNames), bSortModifier);
+			break;
+			case EMaterialInstanceColumnType::ComponentMask:
+				SortByComponentMask(FilteredRows, GetComponentMaskIndexFromColumnId(SortColumnId, *ParameterNames), bSortModifier);
+				break;
 		}
 	}
-	return SNew(SMaterialInstanceTable).InitialRows(Stats);
+};
+
+/** Instantiate a table widget for all material instances in a material. */
+static TSharedPtr<SMaterialInstanceTable> CreateMaterialInstanceTable()
+{
+	TArray<TSharedPtr<FMaterialInstanceColumnData>> Columns;
+	TArray<TSharedPtr<FMaterialInstanceRowData>> Rows;
+	TSharedPtr<FMaterialInstanceParameterNames> ParameterNames = MakeShared<FMaterialInstanceParameterNames>();
+
+	return SNew(SMaterialInstanceTable).InitialColumns(Columns).InitialRows(Rows).InitialParameterNames(ParameterNames);
 }
 
-static void UpdateMaterialInstanceTable(TSharedPtr<SMaterialInstanceTable> InTable, FMaterialValidationDesc const* InMaterialDesc)
+/** Update the table for all material instances in a material. */
+static void UpdateMaterialInstanceTable(TSharedPtr<SMaterialInstanceTable> InTable, UMaterialValidationGroup const* InGroup, FString const& InMaterialPath)
 {
-	TArray<TSharedPtr<MaterialValidation::FMaterialInstanceStats>> Stats;
-	if (InMaterialDesc != nullptr)
+	TArray<FString> MaterialInstancePaths;
+	TArray<FMaterialValidationAssetData_MaterialInstance> MaterialInstanceAssetDatas;
+	TSharedPtr<FMaterialInstanceParameterNames> ParameterNames = MakeShared<FMaterialInstanceParameterNames>();
+	
+	// Get all the information from the material validation database.
+	UMaterialValidationLibrary::GetMaterialInstanceAssetDatas(
+		InGroup, 
+		InMaterialPath, 
+		MaterialInstancePaths, 
+		MaterialInstanceAssetDatas, 
+		ParameterNames->StaticSwitches, 
+		ParameterNames->ComponentMasks);
+
+	// Build all rows.
+	TArray<TSharedPtr<FMaterialInstanceRowData>> Rows;
+	const int32 NumMaterialInstances = MaterialInstancePaths.Num();
+	for (int32 Index = 0; Index < NumMaterialInstances; ++Index)
 	{
-		for (TPair<FString, int32> const& It : InMaterialDesc->MaterialInstances)
+		Rows.Add(MakeShared<FMaterialInstanceRowData>(MaterialInstancePaths[Index], MaterialInstanceAssetDatas[Index]));
+	}
+
+	// Build all columns.
+	TArray<TSharedPtr<FMaterialInstanceColumnData>> Columns;
+	if (Rows.Num())
+	{
+		for (int32 PropertyIndex = 0; PropertyIndex < (int32)EMaterialInstanceStaticProperties::Count; ++PropertyIndex)
 		{
-			const int32 AssetDataIndex = It.Value;
-			const uint32 PermutationHash = InMaterialDesc->MaterialInstanceAssetDatas.IsValidIndex(AssetDataIndex) ? InMaterialDesc->MaterialInstanceAssetDatas[AssetDataIndex].PermutationHash : 0;
-			Stats.Add(MakeShared<MaterialValidation::FMaterialInstanceStats>(It.Key, PermutationHash));
+			Columns.Add(MakeShared<FMaterialInstanceColumnData>(StaticPropertyInfos[PropertyIndex].ColumnId, FText::FromString(StaticPropertyInfos[PropertyIndex].PropertyName), StaticPropertyInfos[PropertyIndex].Type));
+		}
+		for (int32 UsageIndex = 0; UsageIndex < MATUSAGE_MAX; ++UsageIndex)
+		{
+			Columns.Add(MakeShared<FMaterialInstanceColumnData>(UsageColumnIds[UsageIndex], FText::FromString(UMaterialInterface::GetUsageName((EMaterialUsage)UsageIndex)), EMaterialInstanceColumnType::UsageFlag));
+		}
+		for (FName& Name : ParameterNames->StaticSwitches)
+		{
+			Columns.Add(MakeShared<FMaterialInstanceColumnData>(Name, FText::FromName(Name), EMaterialInstanceColumnType::StaticSwitch));
+		}
+		for (FName& Name : ParameterNames->ComponentMasks)
+		{
+			Columns.Add(MakeShared<FMaterialInstanceColumnData>(Name, FText::FromName(Name), EMaterialInstanceColumnType::ComponentMask));
 		}
 	}
-	InTable->SetRows(Stats);
+	
+	// Apply to the table.
+	InTable->SetContent(Rows, Columns, ParameterNames);
 }
 
 } // namespace MaterialValidation
@@ -498,6 +1066,11 @@ class FMaterialValidationToolkit : public FAssetEditorToolkit, public FGCObject,
 	/** Material instance table widget. */
 	TSharedPtr<MaterialValidation::SMaterialInstanceTable> MaterialInstanceTable;
 
+	// Toolkit tabs are referenced by FName.
+	static const FName TabId_Details;
+	static const FName TabId_Materials;
+	static const FName TabId_MaterialInstances;
+
 public:
 	// Begin FGCObject overrides.
 	virtual FString GetReferencerName() const override
@@ -536,17 +1109,17 @@ public:
 	{
 		WorkspaceMenuCategory = InTabManager->AddLocalWorkspaceMenuCategory(LOCTEXT("Toolkit_WorkspaceMenuCatgory", "Material Validation"));
 
-		InTabManager->RegisterTabSpawner(MaterialValidation::TabId_Details, FOnSpawnTab::CreateSP(this, &FMaterialValidationToolkit::SpawnDetailsTab))
+		InTabManager->RegisterTabSpawner(TabId_Details, FOnSpawnTab::CreateSP(this, &FMaterialValidationToolkit::SpawnDetailsTab))
 			.SetDisplayName(LOCTEXT("Toolkit_DetailsTab_DisplayName", "Details"))
 			.SetGroup(WorkspaceMenuCategory.ToSharedRef())
 			.SetIcon(FSlateIcon(FAppStyle::GetAppStyleSetName(), "LevelEditor.Tabs.Details"));
 
-		InTabManager->RegisterTabSpawner(MaterialValidation::TabId_Materials, FOnSpawnTab::CreateSP(this, &FMaterialValidationToolkit::SpawnMaterialsTab))
+		InTabManager->RegisterTabSpawner(TabId_Materials, FOnSpawnTab::CreateSP(this, &FMaterialValidationToolkit::SpawnMaterialsTab))
 			.SetDisplayName(LOCTEXT("Toolkit_MaterialsTab_DisplayName", "Materials"))
 			.SetGroup(WorkspaceMenuCategory.ToSharedRef())
 			.SetIcon(FSlateIcon(FAppStyle::GetAppStyleSetName(), "LevelEditor.Tabs.StatsViewer"));
 
-		InTabManager->RegisterTabSpawner(MaterialValidation::TabId_MaterialInstances, FOnSpawnTab::CreateSP(this, &FMaterialValidationToolkit::SpawnMaterialInstancesTab))
+		InTabManager->RegisterTabSpawner(TabId_MaterialInstances, FOnSpawnTab::CreateSP(this, &FMaterialValidationToolkit::SpawnMaterialInstancesTab))
 			.SetDisplayName(LOCTEXT("Toolkit_MaterialInstancesTab_DisplayName", "Material Instances"))
 			.SetGroup(WorkspaceMenuCategory.ToSharedRef())
 			.SetIcon(FSlateIcon(FAppStyle::GetAppStyleSetName(), "LevelEditor.Tabs.StatsViewer"));
@@ -555,9 +1128,9 @@ public:
 	virtual void UnregisterTabSpawners(const TSharedRef<class FTabManager>& InTabManager) override
 	{
 		FAssetEditorToolkit::UnregisterTabSpawners(InTabManager);
-		InTabManager->UnregisterTabSpawner(MaterialValidation::TabId_Details);
-		InTabManager->UnregisterTabSpawner(MaterialValidation::TabId_Materials);
-		InTabManager->UnregisterTabSpawner(MaterialValidation::TabId_MaterialInstances);
+		InTabManager->UnregisterTabSpawner(TabId_Details);
+		InTabManager->UnregisterTabSpawner(TabId_Materials);
+		InTabManager->UnregisterTabSpawner(TabId_MaterialInstances);
 	}
 	// End FAssetEditorToolkit overrides.
 
@@ -603,9 +1176,9 @@ public:
 
 		MaterialTable = MaterialValidation::CreateMaterialTable(EditGroup, MaterialValidation::FOnRowItemSelect::CreateSP(this, &FMaterialValidationToolkit::OnRowItemSelect));
 
-		MaterialInstanceTable = MaterialValidation::CreateMaterialInstanceTable(nullptr);
+		MaterialInstanceTable = MaterialValidation::CreateMaterialInstanceTable();
 
-		const TSharedRef<FTabManager::FLayout> Layout = FTabManager::NewLayout("MaterialValidationGroupLayout_v3")
+		const TSharedRef<FTabManager::FLayout> Layout = FTabManager::NewLayout("MaterialValidationGroupLayout_v4")
 			->AddArea
 			(
 				FTabManager::NewPrimaryArea()->SetOrientation(Orient_Horizontal)
@@ -613,12 +1186,12 @@ public:
 				(
 					FTabManager::NewSplitter()->SetOrientation(Orient_Vertical)
 					->Split(FTabManager::NewStack()
-						->AddTab(MaterialValidation::TabId_Details, ETabState::OpenedTab)
-						->SetForegroundTab(MaterialValidation::TabId_Details))
+						->AddTab(TabId_Details, ETabState::OpenedTab)
+						->SetForegroundTab(TabId_Details))
 					->Split(FTabManager::NewStack()
-						->AddTab(MaterialValidation::TabId_Materials, ETabState::OpenedTab))
+						->AddTab(TabId_Materials, ETabState::OpenedTab))
 					->Split(FTabManager::NewStack()
-						->AddTab(MaterialValidation::TabId_MaterialInstances, ETabState::ClosedTab))
+						->AddTab(TabId_MaterialInstances, ETabState::ClosedTab))
 				)
 			);
 
@@ -628,18 +1201,17 @@ public:
 	/** Apply selected material to the material instances table. */
 	void OnRowItemSelect(TSharedPtr<MaterialValidation::FMaterialStats> InMaterialStats)
 	{
-		FMaterialValidationDesc const* MaterialDesc = nullptr;
-		if (EditGroup != nullptr && InMaterialStats.IsValid())
-		{
-			MaterialDesc = EditGroup->Materials.Find(InMaterialStats->AssetPath);
-		}
+		MaterialValidation::UpdateMaterialInstanceTable(MaterialInstanceTable, EditGroup, InMaterialStats->AssetPath);
 
-		MaterialValidation::UpdateMaterialInstanceTable(MaterialInstanceTable, MaterialDesc);
-
-		InvokeTab(MaterialValidation::TabId_MaterialInstances);
+		InvokeTab(TabId_MaterialInstances);
 	}
 };
 
+const FName FMaterialValidationToolkit::TabId_Details("MainTab");
+const FName FMaterialValidationToolkit::TabId_Materials("MaterialsTab");
+const FName FMaterialValidationToolkit::TabId_MaterialInstances("MaterialInstancesTab");
+
+
 FText UAssetDefinition_MaterialValidationGroup::GetAssetDisplayName() const 
 {
 	return LOCTEXT("Asset_DisplayName", "Material Validation Group"); 
@@ -670,7 +1242,6 @@ EAssetCommandResult UAssetDefinition_MaterialValidationGroup::OpenAssets(const F
 	}
 
 	return EAssetCommandResult::Handled;
-
 }
 
 #undef LOCTEXT_NAMESPACE
diff --git a/Engine/Plugins/Editor/MaterialValidation/Source/Private/MaterialValidationToolkit.h b/Engine/Plugins/Editor/MaterialValidation/Source/Private/MaterialValidationToolkit.h
--- a/Engine/Plugins/Editor/MaterialValidation/Source/Private/MaterialValidationToolkit.h
+++ b/Engine/Plugins/Editor/MaterialValidation/Source/Private/MaterialValidationToolkit.h
@@ -17,3 +17,15 @@ class UAssetDefinition_MaterialValidationGroup : public UAssetDefinitionDefault
 	virtual TConstArrayView<FAssetCategoryPath> GetAssetCategories() const override;
 	virtual EAssetCommandResult OpenAssets(const FAssetOpenArgs& OpenArgs) const override;
 };
+
+/** Enumeration of material instance property types used to filter editor toolkit tables. */
+UENUM()
+enum class EMaterialInstanceColumnType
+{
+	Default,
+	StaticProperty,
+	UsageFlag,
+	StaticSwitch,
+	ComponentMask,
+	Count
+};
