Performance improvements for the material validation toolkit UI.
Materials with many static parameters show much faster.
Also added static parameter counts to the material table. For now these counts are approximate upper bounds because we don't store the actual value in the material database.
#rnx
#Synced-CL 50706270

diff --git a/Engine/Plugins/Editor/MaterialValidation/Source/Private/MaterialValidationToolkit.cpp b/Engine/Plugins/Editor/MaterialValidation/Source/Private/MaterialValidationToolkit.cpp
--- a/Engine/Plugins/Editor/MaterialValidation/Source/Private/MaterialValidationToolkit.cpp
+++ b/Engine/Plugins/Editor/MaterialValidation/Source/Private/MaterialValidationToolkit.cpp
@@ -32,18 +32,24 @@ struct FMaterialStats
 	int32 MaterialInstanceCount = 0;
 	int32 StaticPermutationCount = 0;
 	int32 ShaderCount = 0;
+	int32 StaticSwitchCount = 0;
+	int32 ComponentMaskCount = 0;
 
 	// Stat table columns are referenced by FName.
 	static const FName ColumnId_AssetPath;
 	static const FName ColumnId_MaterialInstanceCount;
 	static const FName ColumnId_StaticPermutationCount;
 	static const FName ColumnId_ShaderCount;
+	static const FName ColumnId_StaticSwitchCount;
+	static const FName ColumnId_ComponentMaskCount;
 
-	FMaterialStats(FString const& InPath, int32 InInstanceCount, int32 InPermutationCount, int32 InShaderCount)
+	FMaterialStats(FString const& InPath, int32 InInstanceCount, int32 InPermutationCount, int32 InShaderCount, int32 InStaticSwitchCount, int32 InComponentMaskCount)
 		: AssetPath(InPath)
 		, MaterialInstanceCount(InInstanceCount)
 		, StaticPermutationCount(InPermutationCount)
 		, ShaderCount(InShaderCount)
+		, StaticSwitchCount(InStaticSwitchCount)
+		, ComponentMaskCount(InComponentMaskCount)
 	{
 	}
 };
@@ -52,6 +58,8 @@ const FName FMaterialStats::ColumnId_AssetPath("CID_AssetPath");
 const FName FMaterialStats::ColumnId_MaterialInstanceCount("CID_MaterialInstanceCount");
 const FName FMaterialStats::ColumnId_StaticPermutationCount("CID_StaticPermutationCount");
 const FName FMaterialStats::ColumnId_ShaderCount("CID_ShaderCount");
+const FName FMaterialStats::ColumnId_StaticSwitchCount("CID_StaticSwitchCount");
+const FName FMaterialStats::ColumnId_ComponentMaskCount("CID_ComponentMaskCount");
 
 /** A slate container of a data table entry describing a UMaterial. */
 class SMaterialTableRow : public SMultiColumnTableRow<TSharedPtr<FMaterialStats>>
@@ -91,6 +99,14 @@ public:
 			{
 				return SNew(STextBlock).Text(FText::AsNumber(RowItem->ShaderCount));
 			}
+			if (ColumnName == FMaterialStats::ColumnId_StaticSwitchCount)
+			{
+				return SNew(STextBlock).Text(FText::AsNumber(RowItem->StaticSwitchCount));
+			}
+			if (ColumnName == FMaterialStats::ColumnId_ComponentMaskCount)
+			{
+				return SNew(STextBlock).Text(FText::AsNumber(RowItem->ComponentMaskCount));
+			}
 		}
 		return SNew(STextBlock).Text(FText::FromString("?"));
 	}
@@ -180,6 +196,18 @@ public:
 						.DefaultLabel(LOCTEXT("MaterialShaderCount", "Shader Count"))
 						.OnSort(this, &SMaterialTable::OnColumnSortModeChanged)
 						.FixedWidth(150)
+					+ SHeaderRow::Column(FMaterialStats::ColumnId_StaticSwitchCount)
+						.HAlignHeader(HAlign_Center)
+						.HAlignCell(HAlign_Center)
+						.DefaultLabel(LOCTEXT("MaterialStaticSwitchCount", "Static Switch Count (upper est.)"))
+						.OnSort(this, &SMaterialTable::OnColumnSortModeChanged)
+						.FixedWidth(150)
+					+ SHeaderRow::Column(FMaterialStats::ColumnId_ComponentMaskCount)
+						.HAlignHeader(HAlign_Center)
+						.HAlignCell(HAlign_Center)
+						.DefaultLabel(LOCTEXT("MaterialComponentMaskCount", "Component Mask Count (upper est.)"))
+						.OnSort(this, &SMaterialTable::OnColumnSortModeChanged)
+						.FixedWidth(150)
 				)
 				.ListItemsSource(&FilteredRows)
 				.OnGenerateRow_Lambda([](TSharedPtr<FMaterialStats> InItem, const TSharedRef<STableViewBase>& InOwnerTable)
@@ -267,6 +295,20 @@ public:
 				return bSortModifier ? (A->ShaderCount > B->ShaderCount) : (A->ShaderCount < B->ShaderCount);
 			});
 		}
+		else if (SortColumnId == FMaterialStats::ColumnId_StaticSwitchCount)
+		{
+			FilteredRows.Sort([bSortModifier](TSharedPtr<FMaterialStats> const& A, TSharedPtr<FMaterialStats> const& B)
+			{
+				return bSortModifier ? (A->StaticSwitchCount > B->StaticSwitchCount) : (A->StaticSwitchCount < B->StaticSwitchCount);
+			});
+		}
+		else if (SortColumnId == FMaterialStats::ColumnId_ComponentMaskCount)
+		{
+			FilteredRows.Sort([bSortModifier](TSharedPtr<FMaterialStats> const& A, TSharedPtr<FMaterialStats> const& B)
+			{
+				return bSortModifier ? (A->ComponentMaskCount > B->ComponentMaskCount) : (A->ComponentMaskCount < B->ComponentMaskCount);
+			});
+		}
 	}
 
 	/** Update the view. Must be called after any call to UpdateFilter() or UpdateSort(). */
@@ -287,7 +329,13 @@ static TSharedPtr<SMaterialTable> CreateMaterialTable(UMaterialValidationGroup c
 	{
 		for (TPair<FString, FMaterialValidationDesc> const& It : InGroup->Materials)
 		{
-			Stats.Add(MakeShared<FMaterialStats>(It.Key, It.Value.MaterialInstances.Num(), It.Value.PermutationHashes.Num(), It.Value.NumShadersTotal));
+			Stats.Add(MakeShared<FMaterialStats>(
+				It.Key, 
+				It.Value.MaterialInstances.Num(), 
+				It.Value.PermutationHashes.Num(), 
+				It.Value.NumShadersTotal, 
+				It.Value.AssetData.StaticSwitchValues.Num() * 32,	// For now use upper limit of potential static counts since the actual values
+				It.Value.AssetData.ComponentMaskValues.Num() * 8)); // aren't in the database and require opening the UMaterial.
 		}
 	}
 	return SNew(SMaterialTable).InitialRows(Stats).OnRowItemSelect(OnRowItemSelect);
@@ -296,9 +344,10 @@ static TSharedPtr<SMaterialTable> CreateMaterialTable(UMaterialValidationGroup c
 /** The static switch and static component mask parameter FNames used by a single material. */
 struct FMaterialInstanceParameterNames
 {
-	// todo: Add fast FName->Index lookup because with large arrays, the linear search per table entry is very slow.
 	TArray<FName> StaticSwitches;
+	TMap<FName, int32> StaticSwitchNameToIndex;
 	TArray<FName> ComponentMasks;
+	TMap<FName, int32> ComponentMaskNameToIndex;
 };
 
 /** A description for a single column in the UMaterialInstance data table. */
@@ -541,14 +590,8 @@ static void SortByUsage(TArray<TSharedPtr<FMaterialInstanceRowData>>& InArray, E
 /** Resolve column FName to index of static switch parameter. Returns INDEX_NONE if not found. */
 static int32 GetStaticSwitchIndexFromColumnId(FName InColumnId, FMaterialInstanceParameterNames const& InParameterNames)
 {
-	for (int32 Index = 0; Index < InParameterNames.StaticSwitches.Num(); ++Index)
-	{
-		if (InParameterNames.StaticSwitches[Index] == InColumnId)
-		{
-			return Index;
-		}
-	}
-	return INDEX_NONE;
+	int32 const* IndexPtr = InParameterNames.StaticSwitchNameToIndex.Find(InColumnId);
+	return IndexPtr != nullptr ? *IndexPtr : INDEX_NONE;
 }
 
 /** Sort row data based on the given static switch parameter index. */
@@ -580,14 +623,8 @@ static void SortByStaticSwitch(TArray<TSharedPtr<FMaterialInstanceRowData>>& InA
 /** Resolve column FName to index of static component mask parameter. Returns INDEX_NONE if not found. */
 static int32 GetComponentMaskIndexFromColumnId(FName InColumnId, FMaterialInstanceParameterNames const& InParameterNames)
 {
-	for (int32 Index = 0; Index < InParameterNames.ComponentMasks.Num(); ++Index)
-	{
-		if (InParameterNames.ComponentMasks[Index] == InColumnId)
-		{
-			return Index;
-		}
-	}
-	return INDEX_NONE;
+	int32 const* IndexPtr = InParameterNames.ComponentMaskNameToIndex.Find(InColumnId);
+	return IndexPtr != nullptr ? *IndexPtr : INDEX_NONE;
 }
 
 /** Sort row data based on the given static componenent mask parameter index. */
@@ -691,8 +728,8 @@ class SMaterialInstanceTable : public SCompoundWidget
 {
 	/** All the available columns. */
 	TArray<TSharedPtr<FMaterialInstanceColumnData>> AllColumns;
-	/** An ordered and filtered subset of AllColumns. */
-	TArray<TSharedPtr<FMaterialInstanceColumnData>> FilteredColumns;
+	/** The column ids for the currently hidden subset of AllColumns. */
+	TSet<FName> HiddenColumns;
 
 	/** All the available rows. */
 	TArray<TSharedPtr<FMaterialInstanceRowData>> AllRows;
@@ -731,7 +768,6 @@ public:
 	void Construct(const FArguments& InArgs)
 	{
 		AllColumns = InArgs._InitialColumns;
-		FilteredColumns = InArgs._InitialColumns;
 		AllRows = InArgs._InitialRows;
 		FilteredRows = InArgs._InitialRows;
 
@@ -746,7 +782,6 @@ public:
 		SortMode = EColumnSortMode::Descending;
 
 		HeaderRow = SNew(SHeaderRow);
-		UpdateHeaderRow();
 
 		ChildSlot
 		[
@@ -805,9 +840,9 @@ public:
 		AllColumns = InColumns;
 		ParameterNames = InParameterNames;
 
-		UpdateRowFilter();
-		UpdateColumnFilter();
 		UpdateHeaderRow();
+		UpdateColumnFilter();
+		UpdateRowFilter();
 		UpdateSort();
 		UpdateListView();
 	}
@@ -838,22 +873,25 @@ protected:
 	{
 		if (ListView.IsValid())
 		{
-			ListView->RequestListRefresh();
+			ListView->RebuildList();
 		}
 	}
 
-	/** Apply header row. Must be called after any work that updates filtered columns. */
+	/** Add all columns to header row. */
 	void UpdateHeaderRow()
 	{
-		HeaderRow->ClearColumns();
-		for (TSharedPtr<FMaterialInstanceColumnData> const& Column : FilteredColumns)
+		// Gather all columns and use the batched SHeaderRow::AddColumns() for better performance when there are many columns.
+		TArray<SHeaderRow::FColumn::FArguments> AllArgs;
+		AllArgs.Reserve(AllColumns.Num());
+
+		for (int32 Index = 0; Index < AllColumns.Num(); ++Index)
 		{
-			SHeaderRow::FColumn::FArguments Args = SHeaderRow::Column(Column->Id);
-			Args.DefaultLabel(Column->Label).OnSort(this, &SMaterialInstanceTable::OnColumnSortModeChanged);
+			TSharedPtr<FMaterialInstanceColumnData> const& Column = AllColumns[Index];
+			SHeaderRow::FColumn::FArguments& Args = AllArgs.Add_GetRef(SHeaderRow::Column(Column->Id));
+			
 			if (Column->Id == GetColumnIdFromStaticProperty(EMaterialInstanceStaticProperties::AssetPath))
 			{
 				Args
-				.HAlignHeader(HAlign_Center)
 				.HAlignCell(HAlign_Left)
 				.ManualWidth(GAssetNameColumnWidth)
 				.ManualWidth_Lambda([]() { return GAssetNameColumnWidth; })
@@ -862,12 +900,18 @@ protected:
 			else
 			{
 				Args
-				.HAlignHeader(HAlign_Center)
 				.HAlignCell(HAlign_Center)
 				.FixedWidth(120);
 			}
-			HeaderRow->AddColumn(Args);
+			
+			Args
+			.HAlignHeader(HAlign_Center)
+			.DefaultLabel(Column->Label).OnSort(this, &SMaterialInstanceTable::OnColumnSortModeChanged)
+			.ShouldGenerateWidget(TAttribute<bool>::Create(TAttribute<bool>::FGetter::CreateSP(this, &SMaterialInstanceTable::ShouldShowColumn, Column->Id)));
 		}
+
+		HeaderRow->ClearColumns();
+		HeaderRow->AddColumns(AllArgs);
 	}
 
 	/** Called on selection from EMaterialInstanceColumnType filter combo box. */
@@ -876,7 +920,6 @@ protected:
 		SelectedColumnType = (EMaterialInstanceColumnType)Index;
 
 		UpdateColumnFilter();
-		UpdateHeaderRow();
 		UpdateListView();
 	}
 
@@ -910,23 +953,40 @@ protected:
 		ColumnSearchText = InText.ToString();
 
 		UpdateColumnFilter();
-		UpdateHeaderRow();
 		UpdateListView();
 	}
 
-	/**  Apply filter to AllColumns to recreate FilteredColumns. */
+	/**  Apply filter to AllColumns to recreate HiddenColumns and update HeaderRow. */
 	void UpdateColumnFilter()
 	{
-		FilteredColumns.Reset();
+		HiddenColumns.Reset();
 
-		for (const TSharedPtr<FMaterialInstanceColumnData>& Column : AllColumns)
+		for (int32 Index = 0; Index < AllColumns.Num(); ++Index)
 		{
-			if (Column->ColumnType == EMaterialInstanceColumnType::Default || 
-				(Column->ColumnType == SelectedColumnType && (ColumnSearchText.IsEmpty() || Column->Label.ToString().Contains(ColumnSearchText, ESearchCase::IgnoreCase))))
+			const TSharedPtr<FMaterialInstanceColumnData>& Column = AllColumns[Index];
+
+			if (Column->ColumnType == EMaterialInstanceColumnType::Default)
+			{
+				// Default columns are never hidden.
+				continue;
+			}
+				
+			if (Column->ColumnType == SelectedColumnType && 
+				(ColumnSearchText.IsEmpty() || Column->Label.ToString().Contains(ColumnSearchText, ESearchCase::IgnoreCase)))
 			{
-				FilteredColumns.Add(Column);
+				continue;
 			}
+		
+			HiddenColumns.Add(Column->Id);
 		}
+
+		HeaderRow->RefreshColumns();
+	}
+
+	/** UI delegate to determine if a column should be shown. */
+	bool ShouldShowColumn(FName ColumnId)
+	{
+		return !HiddenColumns.Contains(ColumnId);
 	}
 
 	/** Called when clicking on a column header to sort. */
@@ -950,14 +1010,18 @@ protected:
 	/**  Apply sort to FilteredRows. */
 	void UpdateSort()
 	{
-		TSharedPtr<FMaterialInstanceColumnData>* SortColumnData = FilteredColumns.FindByPredicate([ColumnId = SortColumnId](TSharedPtr<FMaterialInstanceColumnData> const& ColumnData)
+		if (HiddenColumns.Contains(SortColumnId))
+		{
+			// If we are filtering out the sort column then the current behavior is no sorting until a column to sort is next selected.
+			return;
+		}
+
+		TSharedPtr<FMaterialInstanceColumnData>* SortColumnData = AllColumns.FindByPredicate([ColumnId = SortColumnId](TSharedPtr<FMaterialInstanceColumnData> const& ColumnData)
 		{
 			return ColumnData->Id == ColumnId; 
 		});
-		if (SortColumnData == nullptr)
+		if (!ensure(SortColumnData != nullptr))
 		{
-			// This can happen if we are filtering out the sort column.
-			// Current behavior is no sorting until a column to sort is next selected.
 			return;
 		}
 
@@ -1039,6 +1103,16 @@ static void UpdateMaterialInstanceTable(TSharedPtr<SMaterialInstanceTable> InTab
 		}
 	}
 	
+	// Build FMaterialInstanceParameterNames lookups.
+	for (int32 Index = 0; Index < ParameterNames->StaticSwitches.Num(); ++Index)
+	{
+		ParameterNames->StaticSwitchNameToIndex.Add(ParameterNames->StaticSwitches[Index], Index);
+	}
+	for (int32 Index = 0; Index < ParameterNames->ComponentMasks.Num(); ++Index)
+	{
+		ParameterNames->ComponentMaskNameToIndex.Add(ParameterNames->ComponentMasks[Index], Index);
+	}
+
 	// Apply to the table.
 	InTable->SetContent(Rows, Columns, ParameterNames);
 }
