Material Validation: Fix Material Instance UI not showing the parent material.
#rnx

#Synced-CL 50764949

diff --git a/Engine/Plugins/Editor/MaterialValidation/Source/Private/MaterialValidationLibrary.cpp b/Engine/Plugins/Editor/MaterialValidation/Source/Private/MaterialValidationLibrary.cpp
--- a/Engine/Plugins/Editor/MaterialValidation/Source/Private/MaterialValidationLibrary.cpp
+++ b/Engine/Plugins/Editor/MaterialValidation/Source/Private/MaterialValidationLibrary.cpp
@@ -1278,6 +1278,9 @@ void UMaterialValidationLibrary::GetMaterialInstanceAssetDatas(
 		Root.StaticProperties = MaterialAssetData.StaticProperties;
 		Root.MaterialLayerHash = MaterialAssetData.MaterialLayerHash;
 		Root.PermutationHash = MaterialAssetData.PermutationHash;
+
+		OutMaterialInstancePaths.Add(MaterialPath);
+		OutMaterialInstanceAssetDatas.Add(Root);
 	}
 
 	// Walk the material instance hierarchy.
