MaterialValidation: Add static parameter counts to database.
#rnx
#Synced-CL 50764949

diff --git a/Engine/Plugins/Editor/MaterialValidation/Source/Private/MaterialValidationAssetData.h b/Engine/Plugins/Editor/MaterialValidation/Source/Private/MaterialValidationAssetData.h
--- a/Engine/Plugins/Editor/MaterialValidation/Source/Private/MaterialValidationAssetData.h
+++ b/Engine/Plugins/Editor/MaterialValidation/Source/Private/MaterialValidationAssetData.h
@@ -97,6 +97,10 @@ struct FMaterialValidationAssetData_Material
 	uint32 MaterialLayerHash = 0;
 	UPROPERTY()
 	uint32 PermutationHash = 0;
+	UPROPERTY()
+	uint16 StaticSwitchNum = 0;
+	UPROPERTY()
+	uint16 ComponentMaskNum = 0;
 };
 
 /**
diff --git a/Engine/Plugins/Editor/MaterialValidation/Source/Private/MaterialValidationLibrary.cpp b/Engine/Plugins/Editor/MaterialValidation/Source/Private/MaterialValidationLibrary.cpp
--- a/Engine/Plugins/Editor/MaterialValidation/Source/Private/MaterialValidationLibrary.cpp
+++ b/Engine/Plugins/Editor/MaterialValidation/Source/Private/MaterialValidationLibrary.cpp
@@ -185,6 +185,7 @@ static void GetAssetData(
 		const uint32 StaticSwitchNum = InPermutationData.StaticSwitchParameterInfo.Num();
 		const uint32 StaticSwitchBitArrayDwordNum = (StaticSwitchNum + 31) / 32;
 		OutAssetData.StaticSwitchValues.SetNumZeroed(StaticSwitchBitArrayDwordNum);
+		OutAssetData.StaticSwitchNum = StaticSwitchNum;
 
 		for (uint32 Index = 0; Index < StaticSwitchNum; ++Index)
 		{
@@ -205,6 +206,7 @@ static void GetAssetData(
 		const uint32 ComponentMaskNum = InPermutationData.ComponentMaskParameterInfo.Num();
 		const uint32 ComponentMaskBitArrayDwordNum = ((ComponentMaskNum * 4) + 31) / 32;
 		OutAssetData.ComponentMaskValues.SetNumZeroed(ComponentMaskBitArrayDwordNum);
+		OutAssetData.ComponentMaskNum = ComponentMaskNum;
 
 		for (uint32 Index = 0; Index < ComponentMaskNum; ++Index)
 		{
