Material validation UI. Use 6 decimal places for opacity mask threshold.
Makes it easier to spot difference between 0.3333 and 0.333333.
Change the specialization method for FString return types while we are at it.
Also fix displayed permutation count to include base material permutation.
Also show correct values for static parameter counts now that they are available.
#rnx
#Synced-CL 50764949

diff --git a/Engine/Plugins/Editor/MaterialValidation/Source/Private/MaterialValidationToolkit.cpp b/Engine/Plugins/Editor/MaterialValidation/Source/Private/MaterialValidationToolkit.cpp
--- a/Engine/Plugins/Editor/MaterialValidation/Source/Private/MaterialValidationToolkit.cpp
+++ b/Engine/Plugins/Editor/MaterialValidation/Source/Private/MaterialValidationToolkit.cpp
@@ -199,13 +199,13 @@ public:
 					+ SHeaderRow::Column(FMaterialStats::ColumnId_StaticSwitchCount)
 						.HAlignHeader(HAlign_Center)
 						.HAlignCell(HAlign_Center)
-						.DefaultLabel(LOCTEXT("MaterialStaticSwitchCount", "Static Switch Count (upper est.)"))
+						.DefaultLabel(LOCTEXT("MaterialStaticSwitchCount", "Static Switch Count"))
 						.OnSort(this, &SMaterialTable::OnColumnSortModeChanged)
 						.FixedWidth(150)
 					+ SHeaderRow::Column(FMaterialStats::ColumnId_ComponentMaskCount)
 						.HAlignHeader(HAlign_Center)
 						.HAlignCell(HAlign_Center)
-						.DefaultLabel(LOCTEXT("MaterialComponentMaskCount", "Component Mask Count (upper est.)"))
+						.DefaultLabel(LOCTEXT("MaterialComponentMaskCount", "Component Mask Count"))
 						.OnSort(this, &SMaterialTable::OnColumnSortModeChanged)
 						.FixedWidth(150)
 				)
@@ -332,10 +332,10 @@ static TSharedPtr<SMaterialTable> CreateMaterialTable(UMaterialValidationGroup c
 			Stats.Add(MakeShared<FMaterialStats>(
 				It.Key, 
 				It.Value.MaterialInstances.Num(), 
-				It.Value.PermutationHashes.Num(), 
+				It.Value.PermutationHashes.Num() + 1,	// PermutationHashes doesn't include hash for base material 
 				It.Value.NumShadersTotal, 
-				It.Value.AssetData.StaticSwitchValues.Num() * 32,	// For now use upper limit of potential static counts since the actual values
-				It.Value.AssetData.ComponentMaskValues.Num() * 8)); // aren't in the database and require opening the UMaterial.
+				It.Value.AssetData.StaticSwitchNum,
+				It.Value.AssetData.ComponentMaskNum));
 		}
 	}
 	return SNew(SMaterialTable).InitialRows(Stats).OnRowItemSelect(OnRowItemSelect);
@@ -438,13 +438,20 @@ static void SortByStaticProperty(TArray<TSharedPtr<FMaterialInstanceRowData>>& I
 using PropertySortFunction = void(*)(TArray<TSharedPtr<FMaterialInstanceRowData>>&, bool);
 
 /** Templatized table entry widget generator for each of the EMaterialInstanceStaticProperties. */
-template <auto GetterPtr, bool bFromString = false>
+template <auto GetterPtr>
 static TSharedRef<SWidget> GenerateWidgetForStaticProperty(TSharedPtr<FMaterialInstanceRowData> const& RowItem)
 {
-	if constexpr(bFromString)
+	using ReturnType = TInvokeResult_T<decltype(GetterPtr), FMaterialInstanceRowData const&>;
+
+	if constexpr(std::is_same<std::remove_cv_t<std::remove_reference_t<ReturnType>>, FString>::value)
 	{
 		return SNew(STextBlock).Text(FText::FromString((*GetterPtr)(*RowItem)));
 	}
+	else if constexpr(TIsFloatingPoint<ReturnType>::Value)
+	{
+		static const FNumberFormattingOptions FormatOptions = FNumberFormattingOptions().SetMinimumFractionalDigits(6);
+		return SNew(STextBlock).Text(FText::AsNumber((*GetterPtr)(*RowItem), &FormatOptions));
+	}
 	else
 	{
 		return SNew(STextBlock).Text(FText::AsNumber((*GetterPtr)(*RowItem)));
@@ -465,7 +472,7 @@ struct FMaterialInstanceStaticPropertyInfo
 
 /** Array for table generation data for EMaterialInstanceStaticProperties. */
 static const FMaterialInstanceStaticPropertyInfo StaticPropertyInfos[(int32)EMaterialInstanceStaticProperties::Count] = {
-	{EMaterialInstanceColumnType::Default, "CID_AssetPath", "Asset Path", &SortByStaticProperty<&FMaterialInstanceStaticPropertyGetters::GetAssetPath>, &GenerateWidgetForStaticProperty<&FMaterialInstanceStaticPropertyGetters::GetAssetPath, true> },
+	{EMaterialInstanceColumnType::Default, "CID_AssetPath", "Asset Path", &SortByStaticProperty<&FMaterialInstanceStaticPropertyGetters::GetAssetPath>, &GenerateWidgetForStaticProperty<&FMaterialInstanceStaticPropertyGetters::GetAssetPath> },
 	{EMaterialInstanceColumnType::Default, "CID_PermutationHash", "Permutation Hash", &SortByStaticProperty<&FMaterialInstanceStaticPropertyGetters::GetPermutationHash>, &GenerateWidgetForStaticProperty<&FMaterialInstanceStaticPropertyGetters::GetPermutationHash> },
 	{EMaterialInstanceColumnType::StaticProperty, "CID_TwoSided", "Two Sided", &SortByStaticProperty<&FMaterialInstanceStaticPropertyGetters::GetTwoSided>, &GenerateWidgetForStaticProperty<&FMaterialInstanceStaticPropertyGetters::GetTwoSided> },
 	{EMaterialInstanceColumnType::StaticProperty, "CID_IsThinSurface", "Thin Surface", &SortByStaticProperty<&FMaterialInstanceStaticPropertyGetters::GetIsThinSurface>, &GenerateWidgetForStaticProperty<&FMaterialInstanceStaticPropertyGetters::GetIsThinSurface> },
