Material Validation: Track changed settings from children/parent and highlight in Yellow
When switching base Materials reset column sort if it no longer applies

#tests local editor
#rnx

#Synced-CL 50783381

#Synced-CL 50783381

diff --git a/Engine/Plugins/Editor/MaterialValidation/Source/Private/MaterialValidationLibrary.cpp b/Engine/Plugins/Editor/MaterialValidation/Source/Private/MaterialValidationLibrary.cpp
--- a/Engine/Plugins/Editor/MaterialValidation/Source/Private/MaterialValidationLibrary.cpp
+++ b/Engine/Plugins/Editor/MaterialValidation/Source/Private/MaterialValidationLibrary.cpp
@@ -334,15 +334,17 @@ static uint32 BuildPermutationHash(UMaterialInstanceConstant const& InMaterialIn
 }
 
 /** Apply overrides for bit mask arrays such as those used by static switch and component masks. */
-static void ApplyBitMaskOverrides(TConstArrayView<uint32> InParentValues, TConstArrayView<uint32> InChildOverridesMasks, TConstArrayView<uint32> InChildOverrideValues, TArray<uint32>& OutValues)
+static void ApplyBitMaskOverrides(TConstArrayView<uint32> InParentValues, TConstArrayView<uint32> InChildOverridesMasks, TConstArrayView<uint32> InChildOverrideValues, TArray<uint32>& OutValues, TArray<uint32>& OutModified)
 {
 	const int32 NumDwords = InParentValues.Num();
 	check(NumDwords == InChildOverridesMasks.Num() && NumDwords == InChildOverrideValues.Num());
 	OutValues.SetNumUninitialized(NumDwords);
+	OutModified.SetNumUninitialized(NumDwords);
 
 	for (int32 Index = 0; Index < NumDwords; ++Index)
 	{
 		OutValues[Index] = (InParentValues[Index] & ~InChildOverridesMasks[Index]) | (InChildOverridesMasks[Index] & InChildOverrideValues[Index]);
+		OutModified[Index] = OutValues[Index] ^ InParentValues[Index];
 	}
 }
 
@@ -351,7 +353,8 @@ static void ApplyBasePropertyOverrides(
 	FStaticPermutationProperties const& InParentProperties, 
 	FStaticPermutationPropertyOverrideFlags const& InChildOverrideFlags,
 	FStaticPermutationProperties const& InChildOverrideProperties,
-	FStaticPermutationProperties& OutProperties)
+	FStaticPermutationProperties& OutProperties,
+	FStaticPermutationPropertyOverrideFlags& OutPropertiesModified)
 {
 	OutProperties.TwoSided = InChildOverrideFlags.bOverride_TwoSided ? InChildOverrideProperties.TwoSided : InParentProperties.TwoSided;
 	OutProperties.bIsThinSurface = InChildOverrideFlags.bOverride_bIsThinSurface ? InChildOverrideProperties.bIsThinSurface : InParentProperties.bIsThinSurface;
@@ -365,6 +368,19 @@ static void ApplyBasePropertyOverrides(
 	OutProperties.ShadingModel = InChildOverrideFlags.bOverride_ShadingModel && InChildOverrideProperties.ShadingModel != MSM_FromMaterialExpression ? InChildOverrideProperties.ShadingModel : InParentProperties.ShadingModel;
 	OutProperties.OpacityMaskClipValue = InChildOverrideFlags.bOverride_OpacityMaskClipValue ? InChildOverrideProperties.OpacityMaskClipValue : InParentProperties.OpacityMaskClipValue;
 	OutProperties.UsageFlags = (InParentProperties.UsageFlags & ~InChildOverrideFlags.Override_UsageFlags) | (InChildOverrideFlags.Override_UsageFlags & InChildOverrideProperties.UsageFlags);
+
+	// Return what child properties are actually different
+	OutPropertiesModified.bOverride_OpacityMaskClipValue = OutProperties.OpacityMaskClipValue != InParentProperties.OpacityMaskClipValue;
+	OutPropertiesModified.bOverride_BlendMode = OutProperties.BlendMode != InParentProperties.BlendMode;
+	OutPropertiesModified.bOverride_ShadingModel = OutProperties.ShadingModel != InParentProperties.ShadingModel;
+	OutPropertiesModified.bOverride_DitheredLODTransition = OutProperties.DitheredLODTransition != InParentProperties.DitheredLODTransition;
+	OutPropertiesModified.bOverride_CastDynamicShadowAsMasked = OutProperties.bCastDynamicShadowAsMasked != InParentProperties.bCastDynamicShadowAsMasked;
+	OutPropertiesModified.bOverride_TwoSided = OutProperties.TwoSided != InParentProperties.TwoSided;
+	OutPropertiesModified.bOverride_bIsThinSurface = OutProperties.bIsThinSurface != InParentProperties.bIsThinSurface;
+	OutPropertiesModified.bOverride_OutputTranslucentVelocity = OutProperties.bOutputTranslucentVelocity != InParentProperties.bOutputTranslucentVelocity;
+	OutPropertiesModified.bOverride_bHasPixelAnimation = OutProperties.bHasPixelAnimation != InParentProperties.bHasPixelAnimation;
+	OutPropertiesModified.bOverride_bEnableTessellation = OutProperties.bEnableTessellation != InParentProperties.bEnableTessellation;
+	OutPropertiesModified.Override_UsageFlags = OutProperties.UsageFlags ^ InParentProperties.UsageFlags;
 }
 
 /** Gather all search and exclude directory paths for a group. */
@@ -1193,9 +1209,9 @@ int32 UMaterialValidationLibrary::GetModifiedShaderCount(UMaterialValidationGrou
 
 		Child.StaticSwitchOverrideValues.SetNumZeroed(Parent.StaticSwitchOverrideValues.Num());
 		Child.ComponentMaskOverrideValues.SetNumZeroed(Parent.ComponentMaskOverrideValues.Num());
-		MaterialValidationLibrary::ApplyBitMaskOverrides(Parent.StaticSwitchOverrideValues, AssetData.StaticSwitchOverrideMask, AssetData.StaticSwitchOverrideValues, Child.StaticSwitchOverrideValues);
-		MaterialValidationLibrary::ApplyBitMaskOverrides(Parent.ComponentMaskOverrideValues, AssetData.ComponentMaskOverrideMask, AssetData.ComponentMaskOverrideValues, Child.ComponentMaskOverrideValues);
-		MaterialValidationLibrary::ApplyBasePropertyOverrides(Parent.StaticProperties, AssetData.StaticPropertyOverrideFlags, AssetData.StaticProperties, Child.StaticProperties);
+		MaterialValidationLibrary::ApplyBitMaskOverrides(Parent.StaticSwitchOverrideValues, AssetData.StaticSwitchOverrideMask, AssetData.StaticSwitchOverrideValues, Child.StaticSwitchOverrideValues, Child.StaticSwitchOverrideMask);
+		MaterialValidationLibrary::ApplyBitMaskOverrides(Parent.ComponentMaskOverrideValues, AssetData.ComponentMaskOverrideMask, AssetData.ComponentMaskOverrideValues, Child.ComponentMaskOverrideValues, Child.ComponentMaskOverrideMask);
+		MaterialValidationLibrary::ApplyBasePropertyOverrides(Parent.StaticProperties, AssetData.StaticPropertyOverrideFlags, AssetData.StaticProperties, Child.StaticProperties, Child.StaticPropertyOverrideFlags);
 		Child.MaterialLayerHash = AssetData.MaterialLayerHash;
 
 		// Build the permutation hash and if this is a new permutation then add its shader count.
@@ -1330,9 +1346,9 @@ void UMaterialValidationLibrary::GetMaterialInstanceAssetDatas(
 
 		Child.StaticSwitchOverrideValues.SetNumZeroed(Parent.StaticSwitchOverrideValues.Num());
 		Child.ComponentMaskOverrideValues.SetNumZeroed(Parent.ComponentMaskOverrideValues.Num());
-		MaterialValidationLibrary::ApplyBitMaskOverrides(Parent.StaticSwitchOverrideValues, AssetData.StaticSwitchOverrideMask, AssetData.StaticSwitchOverrideValues, Child.StaticSwitchOverrideValues);
-		MaterialValidationLibrary::ApplyBitMaskOverrides(Parent.ComponentMaskOverrideValues, AssetData.ComponentMaskOverrideMask, AssetData.ComponentMaskOverrideValues, Child.ComponentMaskOverrideValues);
-		MaterialValidationLibrary::ApplyBasePropertyOverrides(Parent.StaticProperties, AssetData.StaticPropertyOverrideFlags, AssetData.StaticProperties, Child.StaticProperties);
+		MaterialValidationLibrary::ApplyBitMaskOverrides(Parent.StaticSwitchOverrideValues, AssetData.StaticSwitchOverrideMask, AssetData.StaticSwitchOverrideValues, Child.StaticSwitchOverrideValues, Child.StaticSwitchOverrideMask);
+		MaterialValidationLibrary::ApplyBitMaskOverrides(Parent.ComponentMaskOverrideValues, AssetData.ComponentMaskOverrideMask, AssetData.ComponentMaskOverrideValues, Child.ComponentMaskOverrideValues, Child.ComponentMaskOverrideMask);
+		MaterialValidationLibrary::ApplyBasePropertyOverrides(Parent.StaticProperties, AssetData.StaticPropertyOverrideFlags, AssetData.StaticProperties, Child.StaticProperties, Child.StaticPropertyOverrideFlags);
 		Child.MaterialLayerHash = AssetData.MaterialLayerHash;
 		Child.PermutationHash = MaterialValidationLibrary::BuildPermutationHash(Child.StaticProperties, Child.StaticSwitchOverrideValues, Child.ComponentMaskOverrideValues, Child.MaterialLayerHash);
 
diff --git a/Engine/Plugins/Editor/MaterialValidation/Source/Private/MaterialValidationToolkit.cpp b/Engine/Plugins/Editor/MaterialValidation/Source/Private/MaterialValidationToolkit.cpp
--- a/Engine/Plugins/Editor/MaterialValidation/Source/Private/MaterialValidationToolkit.cpp
+++ b/Engine/Plugins/Editor/MaterialValidation/Source/Private/MaterialValidationToolkit.cpp
@@ -421,6 +421,22 @@ struct FMaterialInstanceStaticPropertyGetters
 	static float GetOpacityMaskClipValue(FMaterialInstanceRowData const& P) { return P.AssetData.StaticProperties.OpacityMaskClipValue; }
 };
 
+struct FMaterialInstanceStaticPropertyModifiedGetters
+{
+	static bool GetAssetPathModified(FMaterialInstanceRowData const& P) { return false; }
+	static bool GetPermutationHashModified(FMaterialInstanceRowData const& P) { return false; }
+	static bool GetTwoSidedModified(FMaterialInstanceRowData const& P) { return P.AssetData.StaticPropertyOverrideFlags.bOverride_TwoSided; }
+	static bool GetIsThinSurfaceModified(FMaterialInstanceRowData const& P) { return P.AssetData.StaticPropertyOverrideFlags.bOverride_bIsThinSurface; }
+	static bool GetDitheredLODTransitionModified(FMaterialInstanceRowData const& P) { return P.AssetData.StaticPropertyOverrideFlags.bOverride_DitheredLODTransition; }
+	static bool GetCastDynamicShadowAsMaskedModified(FMaterialInstanceRowData const& P) { return P.AssetData.StaticPropertyOverrideFlags.bOverride_CastDynamicShadowAsMasked; }
+	static bool GetOutputTranslucentVelocityModified(FMaterialInstanceRowData const& P) { return P.AssetData.StaticPropertyOverrideFlags.bOverride_OutputTranslucentVelocity; }
+	static bool GetHasPixelAnimationModified(FMaterialInstanceRowData const& P) { return P.AssetData.StaticPropertyOverrideFlags.bOverride_bHasPixelAnimation; }
+	static bool GetEnableTessellationModified(FMaterialInstanceRowData const& P) { return P.AssetData.StaticPropertyOverrideFlags.bOverride_bEnableTessellation; }
+	static bool GetBlendModeModified(FMaterialInstanceRowData const& P) { return P.AssetData.StaticPropertyOverrideFlags.bOverride_BlendMode; }
+	static bool GetShadingModelModified(FMaterialInstanceRowData const& P) { return P.AssetData.StaticPropertyOverrideFlags.bOverride_ShadingModel; }
+	static bool GetOpacityMaskClipModified(FMaterialInstanceRowData const& P) { return P.AssetData.StaticPropertyOverrideFlags.bOverride_OpacityMaskClipValue; }
+};
+
 /** Templatized sort functor for each of the EMaterialInstanceStaticProperties. */
 template <auto GetterPtr>
 static void SortByStaticProperty(TArray<TSharedPtr<FMaterialInstanceRowData>>& InArray, bool bInSortModifer)
@@ -437,8 +453,13 @@ static void SortByStaticProperty(TArray<TSharedPtr<FMaterialInstanceRowData>>& I
 
 using PropertySortFunction = void(*)(TArray<TSharedPtr<FMaterialInstanceRowData>>&, bool);
 
+static FSlateColor GetColor(bool bModified)
+{
+	return bModified ? FSlateColor(EStyleColor::AccentYellow) : FSlateColor(EStyleColor::Foreground);
+}
+
 /** Templatized table entry widget generator for each of the EMaterialInstanceStaticProperties. */
-template <auto GetterPtr>
+template <auto GetterPtr, auto ModifiedGetterPtr>
 static TSharedRef<SWidget> GenerateWidgetForStaticProperty(TSharedPtr<FMaterialInstanceRowData> const& RowItem)
 {
 	using ReturnType = TInvokeResult_T<decltype(GetterPtr), FMaterialInstanceRowData const&>;
@@ -449,12 +470,12 @@ static TSharedRef<SWidget> GenerateWidgetForStaticProperty(TSharedPtr<FMaterialI
 	}
 	else if constexpr(TIsFloatingPoint<ReturnType>::Value)
 	{
-		static const FNumberFormattingOptions FormatOptions = FNumberFormattingOptions().SetMinimumFractionalDigits(6);
-		return SNew(STextBlock).Text(FText::AsNumber((*GetterPtr)(*RowItem), &FormatOptions));
+		static const FNumberFormattingOptions FormatOptions = FNumberFormattingOptions().SetMaximumFractionalDigits(6);
+		return SNew(STextBlock).Text(FText::AsNumber((*GetterPtr)(*RowItem), &FormatOptions)).ColorAndOpacity(GetColor((*ModifiedGetterPtr)(*RowItem)));;
 	}
 	else
 	{
-		return SNew(STextBlock).Text(FText::AsNumber((*GetterPtr)(*RowItem)));
+		return SNew(STextBlock).Text(FText::AsNumber((*GetterPtr)(*RowItem))).ColorAndOpacity(GetColor((*ModifiedGetterPtr)(*RowItem)));;
 	}
 }
 
@@ -472,18 +493,18 @@ struct FMaterialInstanceStaticPropertyInfo
 
 /** Array for table generation data for EMaterialInstanceStaticProperties. */
 static const FMaterialInstanceStaticPropertyInfo StaticPropertyInfos[(int32)EMaterialInstanceStaticProperties::Count] = {
-	{EMaterialInstanceColumnType::Default, "CID_AssetPath", "Asset Path", &SortByStaticProperty<&FMaterialInstanceStaticPropertyGetters::GetAssetPath>, &GenerateWidgetForStaticProperty<&FMaterialInstanceStaticPropertyGetters::GetAssetPath> },
-	{EMaterialInstanceColumnType::Default, "CID_PermutationHash", "Permutation Hash", &SortByStaticProperty<&FMaterialInstanceStaticPropertyGetters::GetPermutationHash>, &GenerateWidgetForStaticProperty<&FMaterialInstanceStaticPropertyGetters::GetPermutationHash> },
-	{EMaterialInstanceColumnType::StaticProperty, "CID_TwoSided", "Two Sided", &SortByStaticProperty<&FMaterialInstanceStaticPropertyGetters::GetTwoSided>, &GenerateWidgetForStaticProperty<&FMaterialInstanceStaticPropertyGetters::GetTwoSided> },
-	{EMaterialInstanceColumnType::StaticProperty, "CID_IsThinSurface", "Thin Surface", &SortByStaticProperty<&FMaterialInstanceStaticPropertyGetters::GetIsThinSurface>, &GenerateWidgetForStaticProperty<&FMaterialInstanceStaticPropertyGetters::GetIsThinSurface> },
-	{EMaterialInstanceColumnType::StaticProperty, "CID_DitheredLODTransition", "Dithered LOD Transition", &SortByStaticProperty<&FMaterialInstanceStaticPropertyGetters::GetDitheredLODTransition>, &GenerateWidgetForStaticProperty<&FMaterialInstanceStaticPropertyGetters::GetDitheredLODTransition> },
-	{EMaterialInstanceColumnType::StaticProperty, "CID_CastDynamicShadowAsMasked", "Dynamic Shadow As Masked", &SortByStaticProperty<&FMaterialInstanceStaticPropertyGetters::GetCastDynamicShadowAsMasked>, &GenerateWidgetForStaticProperty<&FMaterialInstanceStaticPropertyGetters::GetCastDynamicShadowAsMasked> },
-	{EMaterialInstanceColumnType::StaticProperty, "CID_OutputTranslucentVelocity", "Output Translucent Velocity", &SortByStaticProperty<&FMaterialInstanceStaticPropertyGetters::GetOutputTranslucentVelocity>, &GenerateWidgetForStaticProperty<&FMaterialInstanceStaticPropertyGetters::GetOutputTranslucentVelocity> },
-	{EMaterialInstanceColumnType::StaticProperty, "CID_HasPixelAnimation", "Has Pixel Animation", &SortByStaticProperty<&FMaterialInstanceStaticPropertyGetters::GetHasPixelAnimation>, &GenerateWidgetForStaticProperty<&FMaterialInstanceStaticPropertyGetters::GetHasPixelAnimation> },
-	{EMaterialInstanceColumnType::StaticProperty, "CID_EnableTessellation", "Tessellation", &SortByStaticProperty<&FMaterialInstanceStaticPropertyGetters::GetEnableTessellation>, &GenerateWidgetForStaticProperty<&FMaterialInstanceStaticPropertyGetters::GetEnableTessellation> },
-	{EMaterialInstanceColumnType::StaticProperty, "CID_BlendMode", "Blend Mode", &SortByStaticProperty<&FMaterialInstanceStaticPropertyGetters::GetBlendMode>, &GenerateWidgetForStaticProperty<&FMaterialInstanceStaticPropertyGetters::GetBlendMode> },
-	{EMaterialInstanceColumnType::StaticProperty, "CID_ShadingModel", "Shading Model", &SortByStaticProperty<&FMaterialInstanceStaticPropertyGetters::GetShadingModel>, &GenerateWidgetForStaticProperty<&FMaterialInstanceStaticPropertyGetters::GetShadingModel> },
-	{EMaterialInstanceColumnType::StaticProperty, "CID_OpacityMaskClipValue", "Opacity Mask Clip Value", &SortByStaticProperty<&FMaterialInstanceStaticPropertyGetters::GetOpacityMaskClipValue>, &GenerateWidgetForStaticProperty<&FMaterialInstanceStaticPropertyGetters::GetOpacityMaskClipValue> },
+	{EMaterialInstanceColumnType::Default, "CID_AssetPath", "Asset Path", &SortByStaticProperty<&FMaterialInstanceStaticPropertyGetters::GetAssetPath>, &GenerateWidgetForStaticProperty<&FMaterialInstanceStaticPropertyGetters::GetAssetPath, &FMaterialInstanceStaticPropertyModifiedGetters::GetAssetPathModified>},
+	{EMaterialInstanceColumnType::Default, "CID_PermutationHash", "Permutation Hash", &SortByStaticProperty<&FMaterialInstanceStaticPropertyGetters::GetPermutationHash>, &GenerateWidgetForStaticProperty<&FMaterialInstanceStaticPropertyGetters::GetPermutationHash, &FMaterialInstanceStaticPropertyModifiedGetters::GetPermutationHashModified> },
+	{EMaterialInstanceColumnType::StaticProperty, "CID_TwoSided", "Two Sided", &SortByStaticProperty<&FMaterialInstanceStaticPropertyGetters::GetTwoSided>, &GenerateWidgetForStaticProperty<&FMaterialInstanceStaticPropertyGetters::GetTwoSided, &FMaterialInstanceStaticPropertyModifiedGetters::GetTwoSidedModified> },
+	{EMaterialInstanceColumnType::StaticProperty, "CID_IsThinSurface", "Thin Surface", &SortByStaticProperty<&FMaterialInstanceStaticPropertyGetters::GetIsThinSurface>, &GenerateWidgetForStaticProperty<&FMaterialInstanceStaticPropertyGetters::GetIsThinSurface, &FMaterialInstanceStaticPropertyModifiedGetters::GetIsThinSurfaceModified> },
+	{EMaterialInstanceColumnType::StaticProperty, "CID_DitheredLODTransition", "Dithered LOD Transition", &SortByStaticProperty<&FMaterialInstanceStaticPropertyGetters::GetDitheredLODTransition>, &GenerateWidgetForStaticProperty<&FMaterialInstanceStaticPropertyGetters::GetDitheredLODTransition, &FMaterialInstanceStaticPropertyModifiedGetters::GetDitheredLODTransitionModified> },
+	{EMaterialInstanceColumnType::StaticProperty, "CID_CastDynamicShadowAsMasked", "Dynamic Shadow As Masked", &SortByStaticProperty<&FMaterialInstanceStaticPropertyGetters::GetCastDynamicShadowAsMasked>, &GenerateWidgetForStaticProperty<&FMaterialInstanceStaticPropertyGetters::GetCastDynamicShadowAsMasked, &FMaterialInstanceStaticPropertyModifiedGetters::GetCastDynamicShadowAsMaskedModified> },
+	{EMaterialInstanceColumnType::StaticProperty, "CID_OutputTranslucentVelocity", "Output Translucent Velocity", &SortByStaticProperty<&FMaterialInstanceStaticPropertyGetters::GetOutputTranslucentVelocity>, &GenerateWidgetForStaticProperty<&FMaterialInstanceStaticPropertyGetters::GetOutputTranslucentVelocity, &FMaterialInstanceStaticPropertyModifiedGetters::GetOutputTranslucentVelocityModified> },
+	{EMaterialInstanceColumnType::StaticProperty, "CID_HasPixelAnimation", "Has Pixel Animation", &SortByStaticProperty<&FMaterialInstanceStaticPropertyGetters::GetHasPixelAnimation>, &GenerateWidgetForStaticProperty<&FMaterialInstanceStaticPropertyGetters::GetHasPixelAnimation, &FMaterialInstanceStaticPropertyModifiedGetters::GetHasPixelAnimationModified> },
+	{EMaterialInstanceColumnType::StaticProperty, "CID_EnableTessellation", "Tessellation", &SortByStaticProperty<&FMaterialInstanceStaticPropertyGetters::GetEnableTessellation>, &GenerateWidgetForStaticProperty<&FMaterialInstanceStaticPropertyGetters::GetEnableTessellation, &FMaterialInstanceStaticPropertyModifiedGetters::GetEnableTessellationModified> },
+	{EMaterialInstanceColumnType::StaticProperty, "CID_BlendMode", "Blend Mode", &SortByStaticProperty<&FMaterialInstanceStaticPropertyGetters::GetBlendMode>, &GenerateWidgetForStaticProperty<&FMaterialInstanceStaticPropertyGetters::GetBlendMode, &FMaterialInstanceStaticPropertyModifiedGetters::GetBlendModeModified> },
+	{EMaterialInstanceColumnType::StaticProperty, "CID_ShadingModel", "Shading Model", &SortByStaticProperty<&FMaterialInstanceStaticPropertyGetters::GetShadingModel>, &GenerateWidgetForStaticProperty<&FMaterialInstanceStaticPropertyGetters::GetShadingModel, &FMaterialInstanceStaticPropertyModifiedGetters::GetShadingModelModified> },
+	{EMaterialInstanceColumnType::StaticProperty, "CID_OpacityMaskClipValue", "Opacity Mask Clip Value", &SortByStaticProperty<&FMaterialInstanceStaticPropertyGetters::GetOpacityMaskClipValue>, &GenerateWidgetForStaticProperty<&FMaterialInstanceStaticPropertyGetters::GetOpacityMaskClipValue, &FMaterialInstanceStaticPropertyModifiedGetters::GetOpacityMaskClipModified> },
 };
 
 /** Get column FName for a static property enumeration. */
@@ -698,7 +719,8 @@ public:
 			if (Usage != MATUSAGE_MAX)
 			{
 				const bool bValue = (RowItem->AssetData.StaticProperties.UsageFlags & (1u << Usage)) != 0;
-				return SNew(STextBlock).Text(FText::AsNumber(bValue ? 1 : 0));
+				const bool bModified = (RowItem->AssetData.StaticPropertyOverrideFlags.Override_UsageFlags & (1u << Usage)) != 0;
+				return SNew(STextBlock).Text(FText::AsNumber(bValue ? 1 : 0)).ColorAndOpacity(GetColor(bModified));
 			}
 
 			// Try to resolve column as static switch.
@@ -708,7 +730,8 @@ public:
 				const uint32 BitArrayIndex = StaticSwitchIndex / 32;
 				const uint32 BitIndex = StaticSwitchIndex % 32;
 				const bool bValue = (RowItem->AssetData.StaticSwitchOverrideValues[BitArrayIndex] & (1 << BitIndex)) != 0;
-				return SNew(STextBlock).Text(FText::AsNumber(bValue ? 1 : 0));
+				const bool bModified = RowItem->AssetData.StaticSwitchOverrideMask.IsValidIndex(BitArrayIndex) ? (RowItem->AssetData.StaticSwitchOverrideMask[BitArrayIndex] & (1 << BitIndex)) != 0 : false;
+				return SNew(STextBlock).Text(FText::AsNumber(bValue ? 1 : 0)).ColorAndOpacity(GetColor(bModified));
 			}
 
 			// Try to resolve column as static component mask.
@@ -718,12 +741,13 @@ public:
 				const uint32 BitArrayIndex = ComponentMaskIndex / 8;
 				const uint32 BitIndex = (ComponentMaskIndex * 4) % 32;
 				const uint32 Value = RowItem->AssetData.ComponentMaskOverrideValues[BitArrayIndex] >> BitIndex;
+				const uint32 Modified = RowItem->AssetData.ComponentMaskOverrideMask.IsValidIndex(BitArrayIndex) ? (RowItem->AssetData.ComponentMaskOverrideMask[BitArrayIndex] >> BitIndex) & 0xF : 0;
 				TStringBuilder<8> RGBA;
 				RGBA.AppendChar(Value & 1 ? 'R' : '_');
 				RGBA.AppendChar(Value & 2 ? 'G' : '_');
 				RGBA.AppendChar(Value & 4 ? 'B' : '_');
 				RGBA.AppendChar(Value & 8 ? 'A' : '_');
-				return SNew(STextBlock).Text(FText::FromString(*RGBA));
+				return SNew(STextBlock).Text(FText::FromString(*RGBA)).ColorAndOpacity(GetColor(Modified != 0));
 			}
 		}
 		return SNew(STextBlock).Text(FText::FromString("?"));
@@ -847,6 +871,12 @@ public:
 		AllColumns = InColumns;
 		ParameterNames = InParameterNames;
 
+		// Reset sort column if it no longer applies
+		if (FindColumnById(SortColumnId) == nullptr)
+		{
+			SortColumnId = GetColumnIdFromStaticProperty(EMaterialInstanceStaticProperties::PermutationHash);
+		}
+
 		UpdateHeaderRow();
 		UpdateColumnFilter();
 		UpdateRowFilter();
@@ -1014,6 +1044,15 @@ protected:
 		UpdateListView();
 	}
 
+	/**  Find column, if any. */
+	TSharedPtr<FMaterialInstanceColumnData>* FindColumnById(FName Id)
+	{
+		return AllColumns.FindByPredicate([ColumnId = Id](TSharedPtr<FMaterialInstanceColumnData> const& ColumnData)
+				{
+					return ColumnData->Id == ColumnId; 
+				});	
+	}
+
 	/**  Apply sort to FilteredRows. */
 	void UpdateSort()
 	{
@@ -1023,10 +1062,7 @@ protected:
 			return;
 		}
 
-		TSharedPtr<FMaterialInstanceColumnData>* SortColumnData = AllColumns.FindByPredicate([ColumnId = SortColumnId](TSharedPtr<FMaterialInstanceColumnData> const& ColumnData)
-		{
-			return ColumnData->Id == ColumnId; 
-		});
+		TSharedPtr<FMaterialInstanceColumnData>* SortColumnData = FindColumnById(SortColumnId);
 		if (!ensure(SortColumnData != nullptr))
 		{
 			return;
