Material Validation: Updates to visualization
Add uniqueness checkbox
Add column-type hashes and sorting

#rnx
#Synced-CL 50934400

diff --git a/Engine/Plugins/Editor/MaterialValidation/Source/Private/MaterialValidationToolkit.cpp b/Engine/Plugins/Editor/MaterialValidation/Source/Private/MaterialValidationToolkit.cpp
--- a/Engine/Plugins/Editor/MaterialValidation/Source/Private/MaterialValidationToolkit.cpp
+++ b/Engine/Plugins/Editor/MaterialValidation/Source/Private/MaterialValidationToolkit.cpp
@@ -2,6 +2,7 @@
 
 #include "MaterialValidationToolkit.h"
 
+#include "Hash/xxhash.h"
 #include "Materials/MaterialInterface.h"
 #include "MaterialValidationGroup.h"
 #include "MaterialValidationLibrary.h"
@@ -11,6 +12,7 @@
 #include "Toolkits/AssetEditorToolkit.h"
 #include "UObject/GCObject.h"
 #include "Widgets/Docking/SDockTab.h"
+#include "Widgets/Input/SCheckBox.h"
 #include "Widgets/Input/SComboButton.h"
 #include "Widgets/Input/SSearchBox.h"
 #include "Widgets/Layout/SScrollBox.h"
@@ -368,6 +370,20 @@ struct FMaterialInstanceColumnData
 	}
 };
 
+static uint32 GetNoUsageHash(const FStaticPermutationProperties& Properties)
+{
+	FStaticPermutationProperties NoUsageProperties = Properties;
+	NoUsageProperties.UsageFlags = 0;
+	FXxHash64 Hash = FXxHash64::HashBuffer(&NoUsageProperties, sizeof(NoUsageProperties));
+	return (uint32)(Hash.Hash);
+}
+
+static uint32 GetArrayHash(const TArray<uint32>& Array)
+{
+	FXxHash64 Hash = FXxHash64::HashBuffer(Array);
+	return (uint32)(Hash.Hash);
+}
+
 /** A data table entry describing a UMaterialInstance. */
 struct FMaterialInstanceRowData
 {
@@ -376,10 +392,17 @@ struct FMaterialInstanceRowData
 	/** Full resolved asset data description for the UMaterialInstance pulled from the material validation database. */
 	FMaterialValidationAssetData_MaterialInstance AssetData;
 
+	/** Per-column-type hash data, derived from instance data */
+	uint32 PerColumnTypeHash[(int32)EMaterialInstanceColumnType::Count] = {};
+
 	FMaterialInstanceRowData(FString const& InPath, FMaterialValidationAssetData_MaterialInstance const& InAssetData)
 		: AssetPath(InPath)
 		, AssetData(InAssetData)
 	{
+		PerColumnTypeHash[(int32)EMaterialInstanceColumnType::StaticProperty] = GetNoUsageHash(InAssetData.StaticProperties);
+		PerColumnTypeHash[(int32)EMaterialInstanceColumnType::UsageFlag] = InAssetData.StaticProperties.UsageFlags;
+		PerColumnTypeHash[(int32)EMaterialInstanceColumnType::StaticSwitch] = GetArrayHash(InAssetData.StaticSwitchOverrideValues);
+		PerColumnTypeHash[(int32)EMaterialInstanceColumnType::ComponentMask] = GetArrayHash(InAssetData.ComponentMaskOverrideValues);
 	}
 };
 
@@ -475,7 +498,8 @@ static TSharedRef<SWidget> GenerateWidgetForStaticProperty(TSharedPtr<FMaterialI
 	}
 	else
 	{
-		return SNew(STextBlock).Text(FText::AsNumber((*GetterPtr)(*RowItem))).ColorAndOpacity(GetColor((*ModifiedGetterPtr)(*RowItem)));;
+		static const FNumberFormattingOptions FormatOptions = FNumberFormattingOptions().SetUseGrouping(false);
+		return SNew(STextBlock).Text(FText::AsNumber((*GetterPtr)(*RowItem), &FormatOptions)).ColorAndOpacity(GetColor((*ModifiedGetterPtr)(*RowItem)));;
 	}
 }
 
@@ -580,6 +604,15 @@ static const FName UsageColumnIds[MATUSAGE_MAX] = {
 	"CID_UsageMeshDeformer",
 };
 
+/** Static array of column FNames for column type hashes, outside of Default */
+static const FName ColumnTypeHashIds[(int32)(EMaterialInstanceColumnType::Count)] = {
+	"CID_DefaultUnused",
+	"CID_StaticPropertyHash",
+	"CID_UsageFlagHash",
+	"CID_StaticSwitchHash",
+	"CID_ComponentMaskHash",
+};
+
 /** Resolve column FName to material usage. */
 static EMaterialUsage GetMaterialUsageFromColumnId(FName InColumnId)
 {
@@ -615,6 +648,29 @@ static void SortByUsage(TArray<TSharedPtr<FMaterialInstanceRowData>>& InArray, E
 	}
 }
 
+/** Sort row data based on the given derived column type hash index. */
+static void SortByColumnHash(TArray<TSharedPtr<FMaterialInstanceRowData>>& InArray, int32 Index, bool bInSortModifer)
+{
+	if (Index < (int32)EMaterialInstanceColumnType::Count)
+	{
+		if (bInSortModifer)
+		{
+			Algo::StableSort(InArray, [Index](TSharedPtr<FMaterialInstanceRowData> const& A, TSharedPtr<FMaterialInstanceRowData> const& B)
+			{
+				return (A->PerColumnTypeHash[Index]) > (B->PerColumnTypeHash[Index]);
+			});
+		}
+		else
+		{
+			Algo::StableSort(InArray, [Index](TSharedPtr<FMaterialInstanceRowData> const& A, TSharedPtr<FMaterialInstanceRowData> const& B)
+			{
+				return (A->PerColumnTypeHash[Index]) < (B->PerColumnTypeHash[Index]);
+			});
+		}
+	}
+}
+
+
 /** Resolve column FName to index of static switch parameter. Returns INDEX_NONE if not found. */
 static int32 GetStaticSwitchIndexFromColumnId(FName InColumnId, FMaterialInstanceParameterNames const& InParameterNames)
 {
@@ -749,6 +805,16 @@ public:
 				RGBA.AppendChar(Value & 8 ? 'A' : '_');
 				return SNew(STextBlock).Text(FText::FromString(*RGBA)).ColorAndOpacity(GetColor(Modified != 0));
 			}
+
+			// Try to resolve column as a column type hash
+			static const FNumberFormattingOptions FormatOptions = FNumberFormattingOptions().SetUseGrouping(false);
+			for (int32 Index = 0; Index < (int32)EMaterialInstanceColumnType::Count; ++Index)
+			{
+				if (ColumnId == ColumnTypeHashIds[Index])
+				{
+					return SNew(STextBlock).Text(FText::AsNumber(RowItem->PerColumnTypeHash[Index], &FormatOptions));
+				}
+			}
 		}
 		return SNew(STextBlock).Text(FText::FromString("?"));
 	}
@@ -782,6 +848,8 @@ class SMaterialInstanceTable : public SCompoundWidget
 	FString RowSearchText;
 	/** Search text used for filtering the columns to show. */
 	FString ColumnSearchText;
+	/** True if we are filtering to unique instances. */
+	bool bUniqueInstances;
 
 	/** Selected column name for sorting the table. */
 	FName SortColumnId = NAME_None;
@@ -808,6 +876,7 @@ public:
 
 		RowSearchText.Reset();
 		ColumnSearchText.Reset();
+		bUniqueInstances = false;
 
 		SortColumnId = GetColumnIdFromStaticProperty(EMaterialInstanceStaticProperties::PermutationHash);
 		SortMode = EColumnSortMode::Descending;
@@ -830,6 +899,15 @@ public:
 						.Text(this, &SMaterialInstanceTable::GetSelectedColumnTypeText)
 					]
 				]
+				+ SHorizontalBox::Slot().FillWidth(0.1).Padding(2)
+				[
+					SNew(SCheckBox)
+					.OnCheckStateChanged(this, &SMaterialInstanceTable::OnUniqueCheckboxChanged)
+					[
+						SNew(STextBlock)
+						.Text(LOCTEXT("Toolbox_UniqueCheckbox", "Unique"))
+					]
+				]
 				+ SHorizontalBox::Slot().FillWidth(0.4).Padding(2)
 				[
 					SNew(SSearchBox)
@@ -956,6 +1034,7 @@ protected:
 	{
 		SelectedColumnType = (EMaterialInstanceColumnType)Index;
 
+		UpdateRowFilter();
 		UpdateColumnFilter();
 		UpdateListView();
 	}
@@ -970,15 +1049,36 @@ protected:
 		UpdateListView();
 	}
 
+	/** Called when modifying the unique instances checkbox. */
+	void OnUniqueCheckboxChanged(ECheckBoxState NewState)
+	{
+		bUniqueInstances = (NewState == ECheckBoxState::Checked);
+
+		UpdateRowFilter();
+		UpdateSort();
+		UpdateListView();
+	}	
+
 	/**  Apply filter to AllRows to recreate FilteredRows. */
 	void UpdateRowFilter()
 	{
 		FilteredRows.Reset();
 
+		TSet<uint32> PermutationHashesSeen;
+
 		for (const TSharedPtr<FMaterialInstanceRowData>& Row : AllRows)
 		{
 			if (RowSearchText.IsEmpty() || Row->AssetPath.Contains(RowSearchText, ESearchCase::IgnoreCase))
 			{
+				if (bUniqueInstances)
+				{
+					bool bInSet = false;
+					PermutationHashesSeen.Add(Row->PerColumnTypeHash[(uint32)SelectedColumnType], &bInSet);
+					if (bInSet)
+					{
+						continue;
+					}
+				}
 				FilteredRows.Add(Row);
 			}
 		}
@@ -1056,6 +1156,12 @@ protected:
 	/**  Apply sort to FilteredRows. */
 	void UpdateSort()
 	{
+		if (AllColumns.IsEmpty() || AllRows.IsEmpty())
+		{
+			// If the panel is empty, there is nothing to sort
+			return;
+		}
+
 		if (HiddenColumns.Contains(SortColumnId))
 		{
 			// If we are filtering out the sort column then the current behavior is no sorting until a column to sort is next selected.
@@ -1071,6 +1177,17 @@ protected:
 		// Call specific sorting function for the column type.
 		const bool bSortModifier = SortMode == EColumnSortMode::Descending;
 		const EMaterialInstanceColumnType ColumnType = (*SortColumnData)->ColumnType;
+
+		// See if we are sorting from a column type hash
+		for (int32 Index = 0; Index < (int32)EMaterialInstanceColumnType::Count; ++Index)
+		{
+			if (ColumnTypeHashIds[Index] == SortColumnId)
+			{
+				SortByColumnHash(FilteredRows, Index, bSortModifier);
+				return;
+			}
+		}
+
 		switch (ColumnType)
 		{
 			case EMaterialInstanceColumnType::Default:
@@ -1128,7 +1245,20 @@ static void UpdateMaterialInstanceTable(TSharedPtr<SMaterialInstanceTable> InTab
 	TArray<TSharedPtr<FMaterialInstanceColumnData>> Columns;
 	if (Rows.Num())
 	{
-		for (int32 PropertyIndex = 0; PropertyIndex < (int32)EMaterialInstanceStaticProperties::Count; ++PropertyIndex)
+		// Build path and permutation hash so they are first
+		for (int32 PropertyIndex = 0; PropertyIndex < 2; ++PropertyIndex)
+		{
+			Columns.Add(MakeShared<FMaterialInstanceColumnData>(StaticPropertyInfos[PropertyIndex].ColumnId, FText::FromString(StaticPropertyInfos[PropertyIndex].PropertyName), StaticPropertyInfos[PropertyIndex].Type));
+		}
+
+		// Build additional type-specific hash columns
+		Columns.Add(MakeShared<FMaterialInstanceColumnData>(ColumnTypeHashIds[(int32)EMaterialInstanceColumnType::StaticProperty], LOCTEXT("Tookkit_StaticPropertyHash","Static Property Hash"), EMaterialInstanceColumnType::StaticProperty));
+		Columns.Add(MakeShared<FMaterialInstanceColumnData>(ColumnTypeHashIds[(int32)EMaterialInstanceColumnType::UsageFlag], LOCTEXT("Tookkit_UsageHash","Usage Hash"), EMaterialInstanceColumnType::UsageFlag));
+		Columns.Add(MakeShared<FMaterialInstanceColumnData>(ColumnTypeHashIds[(int32)EMaterialInstanceColumnType::StaticSwitch], LOCTEXT("Tookkit_StaticSwitchHash","Static Switch Hash"), EMaterialInstanceColumnType::StaticSwitch));
+		Columns.Add(MakeShared<FMaterialInstanceColumnData>(ColumnTypeHashIds[(int32)EMaterialInstanceColumnType::ComponentMask], LOCTEXT("Tookkit_ComponentMaskHash","Component Mask Hash"), EMaterialInstanceColumnType::ComponentMask));
+
+		// Build the rest of the columns
+		for (int32 PropertyIndex = 2; PropertyIndex < (int32)EMaterialInstanceStaticProperties::Count; ++PropertyIndex)
 		{
 			Columns.Add(MakeShared<FMaterialInstanceColumnData>(StaticPropertyInfos[PropertyIndex].ColumnId, FText::FromString(StaticPropertyInfos[PropertyIndex].PropertyName), StaticPropertyInfos[PropertyIndex].Type));
 		}
