Material Validation: Improve FStaticPermutationProperties hashing

#tests local editor

#Synced-CL 50934400

diff --git a/Engine/Plugins/Editor/MaterialValidation/Source/Private/MaterialValidationAssetData.h b/Engine/Plugins/Editor/MaterialValidation/Source/Private/MaterialValidationAssetData.h
--- a/Engine/Plugins/Editor/MaterialValidation/Source/Private/MaterialValidationAssetData.h
+++ b/Engine/Plugins/Editor/MaterialValidation/Source/Private/MaterialValidationAssetData.h
@@ -3,6 +3,8 @@
 #pragma once
 
 #include "Materials/MaterialInstanceBasePropertyOverrides.h"
+#include "Templates/TypeHash.h"
+
 #include "MaterialValidationAssetData.generated.h"
 
 /** 
@@ -38,6 +40,10 @@ struct FStaticPermutationProperties
 	uint32 UsageFlags;
 
 	FStaticPermutationProperties();
+
+	bool operator==(const FStaticPermutationProperties& Other) const = default;
+
+	friend uint32 GetTypeHash(const FStaticPermutationProperties& Arg);
 };
 
 /**
diff --git a/Engine/Plugins/Editor/MaterialValidation/Source/Private/MaterialValidationGroup.cpp b/Engine/Plugins/Editor/MaterialValidation/Source/Private/MaterialValidationGroup.cpp
--- a/Engine/Plugins/Editor/MaterialValidation/Source/Private/MaterialValidationGroup.cpp
+++ b/Engine/Plugins/Editor/MaterialValidation/Source/Private/MaterialValidationGroup.cpp
@@ -51,6 +51,25 @@ FStaticPermutationProperties::FStaticPermutationProperties()
 {
 }
 
+uint32 GetTypeHash(const FStaticPermutationProperties& Arg)
+{
+	const uint8 PackedBools = 
+		(Arg.TwoSided ? 1 << 0 : 0) |
+		(Arg.bIsThinSurface ? 1 << 1 : 0) |
+		(Arg.DitheredLODTransition ? 1 << 2 : 0) |
+		(Arg.bCastDynamicShadowAsMasked ? 1 << 3 : 0) |
+		(Arg.bOutputTranslucentVelocity ? 1 << 4 : 0) |
+		(Arg.bHasPixelAnimation ? 1 << 5 : 0) |
+		(Arg.bEnableTessellation ? 1 << 6 : 0);
+
+	uint32 Hash = ::GetTypeHash(PackedBools);
+	Hash = HashCombineFast(Hash, ::GetTypeHash(Arg.OpacityMaskClipValue));
+	Hash = HashCombineFast(Hash, ::GetTypeHash(Arg.BlendMode));
+	Hash = HashCombineFast(Hash, ::GetTypeHash(Arg.ShadingModel));
+	Hash = HashCombineFast(Hash, ::GetTypeHash(Arg.UsageFlags));
+	return Hash;
+}
+
 FStaticPermutationPropertyOverrideFlags::FStaticPermutationPropertyOverrideFlags()
 	: bOverride_OpacityMaskClipValue(0)
 	, bOverride_BlendMode(0)
diff --git a/Engine/Plugins/Editor/MaterialValidation/Source/Private/MaterialValidationLibrary.cpp b/Engine/Plugins/Editor/MaterialValidation/Source/Private/MaterialValidationLibrary.cpp
--- a/Engine/Plugins/Editor/MaterialValidation/Source/Private/MaterialValidationLibrary.cpp
+++ b/Engine/Plugins/Editor/MaterialValidation/Source/Private/MaterialValidationLibrary.cpp
@@ -127,26 +127,6 @@ FMaterialInstanceBasePropertyOverrides GetStaticPropertiesAsBasePropertyOverride
 	return BaseProperties;
 }
 
-/** Return a hash built from the FStaticPermutationProperties values. */
-static uint32 BuildStaticPermutationPropertiesHash(FStaticPermutationProperties const& InStaticProperties)
-{
-	const uint8 PackedBools = 
-		(InStaticProperties.TwoSided ? 1 << 0 : 0) |
-		(InStaticProperties.bIsThinSurface ? 1 << 1 : 0) |
-		(InStaticProperties.DitheredLODTransition ? 1 << 2 : 0) |
-		(InStaticProperties.bCastDynamicShadowAsMasked ? 1 << 3 : 0) |
-		(InStaticProperties.bOutputTranslucentVelocity ? 1 << 4 : 0) |
-		(InStaticProperties.bHasPixelAnimation ? 1 << 5 : 0) |
-		(InStaticProperties.bEnableTessellation ? 1 << 6 : 0);
-
-	uint32 Hash = GetTypeHash(PackedBools);
-	Hash = HashCombineFast(Hash, GetTypeHash(InStaticProperties.OpacityMaskClipValue));
-	Hash = HashCombineFast(Hash, GetTypeHash(InStaticProperties.BlendMode));
-	Hash = HashCombineFast(Hash, GetTypeHash(InStaticProperties.ShadingModel));
-	Hash = HashCombineFast(Hash, GetTypeHash(InStaticProperties.UsageFlags));
-	return Hash;
-}
-
 /** Return a hash built from any material layer settings on a material. */
 static uint32 BuildLayerHash(UMaterialInterface const& InMaterial)
 {
@@ -161,7 +141,7 @@ static uint32 BuildLayerHash(UMaterialInterface const& InMaterial)
 /** Return a hash built for a material or material instance which will match for all instances that generate the same shader permutations. */
 static uint32 BuildPermutationHash(FStaticPermutationProperties const& InStaticProperties, TConstArrayView<uint32> InStaticSwitchValues, TConstArrayView<uint32> InComponentMaskValues, uint32 InMaterialLayerHash)
 {
-	uint32 Hash = BuildStaticPermutationPropertiesHash(InStaticProperties);
+	uint32 Hash = GetTypeHash(InStaticProperties);
 	for (uint32 Value : InStaticSwitchValues)
 	{
 		Hash = HashCombineFast(Hash, GetTypeHash(Value));
diff --git a/Engine/Plugins/Editor/MaterialValidation/Source/Private/MaterialValidationToolkit.cpp b/Engine/Plugins/Editor/MaterialValidation/Source/Private/MaterialValidationToolkit.cpp
--- a/Engine/Plugins/Editor/MaterialValidation/Source/Private/MaterialValidationToolkit.cpp
+++ b/Engine/Plugins/Editor/MaterialValidation/Source/Private/MaterialValidationToolkit.cpp
@@ -374,8 +374,7 @@ static uint32 GetNoUsageHash(const FStaticPermutationProperties& Properties)
 {
 	FStaticPermutationProperties NoUsageProperties = Properties;
 	NoUsageProperties.UsageFlags = 0;
-	FXxHash64 Hash = FXxHash64::HashBuffer(&NoUsageProperties, sizeof(NoUsageProperties));
-	return (uint32)(Hash.Hash);
+	return GetTypeHash(NoUsageProperties);
 }
 
 static uint32 GetArrayHash(const TArray<uint32>& Array)
