MaterialValidation: Add Blueprint functions to fetch resolved material properties from the validation database.
These can be used for analysis that isn't just part of the MaterialValidationToolkit.
UMaterialValidationLibrary::GetMaterialHierarchyInfo() can be called to resolve all of the material information.
Then UMaterialValidationLibrary::GetMaterialProperties() and UMaterialValidationLibrary::GetMaterialPropertyValues() can be used on the returned data to access the property values.
#rnx
#Synced-CL 51290350

diff --git a/Engine/Plugins/Editor/MaterialValidation/Source/Private/MaterialValidationLibrary.cpp b/Engine/Plugins/Editor/MaterialValidation/Source/Private/MaterialValidationLibrary.cpp
--- a/Engine/Plugins/Editor/MaterialValidation/Source/Private/MaterialValidationLibrary.cpp
+++ b/Engine/Plugins/Editor/MaterialValidation/Source/Private/MaterialValidationLibrary.cpp
@@ -1216,26 +1216,25 @@ int32 UMaterialValidationLibrary::GetModifiedShaderCount(UMaterialValidationGrou
  * Note there is quite a bit of duplication with GetModifiedShaderCount() in this function. 
  * It might be good to consolidate on a single material instance visitor that can do arbitrary work?
  */
-void UMaterialValidationLibrary::GetMaterialInstanceAssetDatas(
-	UMaterialValidationGroup const* InGroup, 
-	FString const& MaterialPath,
-	TArray<FString>& OutMaterialInstancePaths, 
-	TArray<FMaterialValidationAssetData_MaterialInstance>& OutMaterialInstanceAssetDatas,
-	TArray<FName>& OutStaticSwitchNames,
-	TArray<FName>& OutComponentMaskNames)
+void UMaterialValidationLibrary::GetMaterialHierarchyInfo(UMaterialValidationGroup const* InGroup, FString const& InBaseMaterialPath, FMaterialDatabaseAssetHierarchyInfo& OutInfo)
 {
+	OutInfo.MaterialPaths.Reset();
+	OutInfo.MaterialAssetDatas.Reset();
+	OutInfo.StaticSwitchNames.Reset();
+	OutInfo.ComponentMaskNames.Reset();
+
 	if (InGroup == nullptr)
 	{
 		return;
 	}
 
-	FMaterialValidationDesc const* MaterialDesc = InGroup->Materials.Find(MaterialPath);
+	FMaterialValidationDesc const* MaterialDesc = InGroup->Materials.Find(InBaseMaterialPath);
 	if (MaterialDesc == nullptr)
 	{
 		return;
 	}
 
-	TSoftObjectPtr<UMaterial> MaterialSoftPtr(MaterialValidationLibrary::ConvertStringPathToSoftObjectPath(MaterialPath));
+	TSoftObjectPtr<UMaterial> MaterialSoftPtr(MaterialValidationLibrary::ConvertStringPathToSoftObjectPath(InBaseMaterialPath));
 	UMaterial const* Material = MaterialSoftPtr.LoadSynchronous();
 	if (Material == nullptr)
 	{
@@ -1249,13 +1248,16 @@ void UMaterialValidationLibrary::GetMaterialInstanceAssetDatas(
 	MaterialValidationLibrary::FMaterialPermutationData PermutationData;
 	BuildPermutationData(*Material, PermutationData);
 
+	OutInfo.StaticSwitchNames.Reserve(PermutationData.StaticSwitchParameterInfo.Num());
 	for (FMaterialParameterInfo const& Info : PermutationData.StaticSwitchParameterInfo)
 	{
-		OutStaticSwitchNames.Add(Info.Name);
+		OutInfo.StaticSwitchNames.Add(Info.Name);
 	}
+	
+	OutInfo.ComponentMaskNames.Reserve(PermutationData.ComponentMaskParameterInfo.Num());
 	for (FMaterialParameterInfo const& Info : PermutationData.ComponentMaskParameterInfo)
 	{
-		OutComponentMaskNames.Add(Info.Name);
+		OutInfo.ComponentMaskNames.Add(Info.Name);
 	}
 
 	FMaterialValidationAssetData_Material MaterialAssetData = MaterialDesc->AssetData;
@@ -1277,10 +1279,13 @@ void UMaterialValidationLibrary::GetMaterialInstanceAssetDatas(
 		Root.MaterialLayerHash = MaterialAssetData.MaterialLayerHash;
 		Root.PermutationHash = MaterialAssetData.PermutationHash;
 
-		OutMaterialInstancePaths.Add(MaterialPath);
-		OutMaterialInstanceAssetDatas.Add(Root);
+		OutInfo.MaterialPaths.Add(InBaseMaterialPath);
+		OutInfo.MaterialAssetDatas.Add(Root);
 	}
 
+	OutInfo.MaterialPaths.Reserve(PackageNames.Num() + 1);
+	OutInfo.MaterialAssetDatas.Reserve(PackageNames.Num() + 1);
+
 	// Walk the material instance hierarchy.
 	for (int32 Index = 1; Index < PackageNames.Num(); ++Index)
 	{
@@ -1332,8 +1337,238 @@ void UMaterialValidationLibrary::GetMaterialInstanceAssetDatas(
 		Child.MaterialLayerHash = AssetData.MaterialLayerHash;
 		Child.PermutationHash = MaterialValidationLibrary::BuildPermutationHash(Child.StaticProperties, Child.StaticSwitchOverrideValues, Child.ComponentMaskOverrideValues, Child.MaterialLayerHash);
 
-		OutMaterialInstancePaths.Add(PackagePath);
-		OutMaterialInstanceAssetDatas.Add(Child);
+		OutInfo.MaterialPaths.Add(PackagePath);
+		OutInfo.MaterialAssetDatas.Add(Child);
+	}
+}
+
+namespace MaterialValidationLibrary {
+
+/**
+ * Enumeration of all of the UMaterialInstance base properties.
+ * Aside from PermutationHash (kept in same enum for simplicity) these are the property overrides that trigger static permutations.
+ */
+enum class EMaterialInstanceStaticProperties : int32
+{
+	PermutationHash,
+	TwoSided,
+	bIsThinSurface,
+	DitheredLODTransition,
+	bCastDynamicShadowAsMasked,
+	bOutputTranslucentVelocity,
+	bHasPixelAnimation,
+	bEnableTessellation,
+	BlendMode,
+	ShadingModel,
+	OpacityMaskClipValue,
+	Count
+};
+
+/** Helper getter functions for each of the EMaterialInstanceStaticProperties. */
+struct FMaterialInstanceStaticPropertyGetters
+{
+	static uint32 GetPermutationHash(FMaterialValidationAssetData_MaterialInstance const& P) { return P.PermutationHash; }
+	static uint8 GetTwoSided(FMaterialValidationAssetData_MaterialInstance const& P) { return P.StaticProperties.TwoSided; }
+	static uint8 GetIsThinSurface(FMaterialValidationAssetData_MaterialInstance const& P) { return P.StaticProperties.bIsThinSurface; }
+	static uint8 GetDitheredLODTransition(FMaterialValidationAssetData_MaterialInstance const& P) { return P.StaticProperties.DitheredLODTransition; }
+	static uint8 GetCastDynamicShadowAsMasked(FMaterialValidationAssetData_MaterialInstance const& P) { return P.StaticProperties.bCastDynamicShadowAsMasked; }
+	static uint8 GetOutputTranslucentVelocity(FMaterialValidationAssetData_MaterialInstance const& P) { return P.StaticProperties.bOutputTranslucentVelocity; }
+	static uint8 GetHasPixelAnimation(FMaterialValidationAssetData_MaterialInstance const& P) { return P.StaticProperties.bHasPixelAnimation; }
+	static uint8 GetEnableTessellation(FMaterialValidationAssetData_MaterialInstance const& P) { return P.StaticProperties.bEnableTessellation; }
+	static TEnumAsByte<EBlendMode> GetBlendMode(FMaterialValidationAssetData_MaterialInstance const& P) { return P.StaticProperties.BlendMode; }
+	static TEnumAsByte<EMaterialShadingModel> GetShadingModel(FMaterialValidationAssetData_MaterialInstance const& P) { return P.StaticProperties.ShadingModel; }
+	static float GetOpacityMaskClipValue(FMaterialValidationAssetData_MaterialInstance const& P) { return P.StaticProperties.OpacityMaskClipValue; }
+};
+
+struct FMaterialInstanceStaticPropertyModifiedGetters
+{
+	static bool GetPermutationHashModified(FMaterialValidationAssetData_MaterialInstance const& P) { return false; }
+	static bool GetTwoSidedModified(FMaterialValidationAssetData_MaterialInstance const& P) { return P.StaticPropertyOverrideFlags.bOverride_TwoSided; }
+	static bool GetIsThinSurfaceModified(FMaterialValidationAssetData_MaterialInstance const& P) { return P.StaticPropertyOverrideFlags.bOverride_bIsThinSurface; }
+	static bool GetDitheredLODTransitionModified(FMaterialValidationAssetData_MaterialInstance const& P) { return P.StaticPropertyOverrideFlags.bOverride_DitheredLODTransition; }
+	static bool GetCastDynamicShadowAsMaskedModified(FMaterialValidationAssetData_MaterialInstance const& P) { return P.StaticPropertyOverrideFlags.bOverride_CastDynamicShadowAsMasked; }
+	static bool GetOutputTranslucentVelocityModified(FMaterialValidationAssetData_MaterialInstance const& P) { return P.StaticPropertyOverrideFlags.bOverride_OutputTranslucentVelocity; }
+	static bool GetHasPixelAnimationModified(FMaterialValidationAssetData_MaterialInstance const& P) { return P.StaticPropertyOverrideFlags.bOverride_bHasPixelAnimation; }
+	static bool GetEnableTessellationModified(FMaterialValidationAssetData_MaterialInstance const& P) { return P.StaticPropertyOverrideFlags.bOverride_bEnableTessellation; }
+	static bool GetBlendModeModified(FMaterialValidationAssetData_MaterialInstance const& P) { return P.StaticPropertyOverrideFlags.bOverride_BlendMode; }
+	static bool GetShadingModelModified(FMaterialValidationAssetData_MaterialInstance const& P) { return P.StaticPropertyOverrideFlags.bOverride_ShadingModel; }
+	static bool GetOpacityMaskClipModified(FMaterialValidationAssetData_MaterialInstance const& P) { return P.StaticPropertyOverrideFlags.bOverride_OpacityMaskClipValue; }
+};
+
+template <auto GetterPtr>
+static FText GenerateTextForStaticProperty(FMaterialValidationAssetData_MaterialInstance const& AssetData)
+{
+	using ReturnType = TInvokeResult_T<decltype(GetterPtr), FMaterialValidationAssetData_MaterialInstance const&>;
+
+	if constexpr (std::is_same<std::remove_cv_t<std::remove_reference_t<ReturnType>>, FString>::value)
+	{
+		return FText::FromString((*GetterPtr)(AssetData));
+	}
+	else if constexpr (TIsFloatingPoint<ReturnType>::Value)
+	{
+		static const FNumberFormattingOptions FormatOptions = FNumberFormattingOptions().SetMaximumFractionalDigits(6);
+		return FText::AsNumber((*GetterPtr)(AssetData), &FormatOptions);
+	}
+	else
+	{
+		static const FNumberFormattingOptions FormatOptions = FNumberFormattingOptions().SetUseGrouping(false);
+		return FText::AsNumber((*GetterPtr)(AssetData), &FormatOptions);
+	}
+}
+
+using FnGetPropertyValueText = FText(*)(FMaterialValidationAssetData_MaterialInstance const&);
+using FnIsPropertyModified = bool(*)(FMaterialValidationAssetData_MaterialInstance const&);
+
+/** Struct containing information needed to generate table entries for EMaterialInstanceStaticProperties. */
+struct FMaterialInstanceStaticPropertyInfo
+{
+	FName PropertyId;
+	FName PropertyName;
+	EMaterialInstanceColumnType Category;
+	FnGetPropertyValueText GetPropertyValueTextFunction;
+	FnIsPropertyModified GetIsPropertyModifiedFunction;
+};
+
+/** Array for table generation data for EMaterialInstanceStaticProperties. */
+static const FMaterialInstanceStaticPropertyInfo StaticPropertyInfos[(int32)EMaterialInstanceStaticProperties::Count] = {
+	{"PID_PermutationHash", "Permutation Hash", EMaterialInstanceColumnType::Default, &GenerateTextForStaticProperty<&FMaterialInstanceStaticPropertyGetters::GetPermutationHash>, &FMaterialInstanceStaticPropertyModifiedGetters::GetPermutationHashModified },
+	{"PID_TwoSided", "Two Sided", EMaterialInstanceColumnType::StaticProperty, &GenerateTextForStaticProperty <&FMaterialInstanceStaticPropertyGetters::GetTwoSided>, &FMaterialInstanceStaticPropertyModifiedGetters::GetTwoSidedModified },
+	{"PID_IsThinSurface", "Thin Surface", EMaterialInstanceColumnType::StaticProperty, &GenerateTextForStaticProperty <&FMaterialInstanceStaticPropertyGetters::GetIsThinSurface>, &FMaterialInstanceStaticPropertyModifiedGetters::GetIsThinSurfaceModified },
+	{"PID_DitheredLODTransition", "Dithered LOD Transition", EMaterialInstanceColumnType::StaticProperty, &GenerateTextForStaticProperty <&FMaterialInstanceStaticPropertyGetters::GetDitheredLODTransition>, &FMaterialInstanceStaticPropertyModifiedGetters::GetDitheredLODTransitionModified },
+	{"PID_CastDynamicShadowAsMasked", "Dynamic Shadow As Masked", EMaterialInstanceColumnType::StaticProperty, &GenerateTextForStaticProperty <&FMaterialInstanceStaticPropertyGetters::GetCastDynamicShadowAsMasked>, &FMaterialInstanceStaticPropertyModifiedGetters::GetCastDynamicShadowAsMaskedModified },
+	{"PID_OutputTranslucentVelocity", "Output Translucent Velocity", EMaterialInstanceColumnType::StaticProperty, &GenerateTextForStaticProperty <&FMaterialInstanceStaticPropertyGetters::GetOutputTranslucentVelocity>, &FMaterialInstanceStaticPropertyModifiedGetters::GetOutputTranslucentVelocityModified },
+	{"PID_HasPixelAnimation", "Has Pixel Animation", EMaterialInstanceColumnType::StaticProperty, &GenerateTextForStaticProperty <&FMaterialInstanceStaticPropertyGetters::GetHasPixelAnimation>, &FMaterialInstanceStaticPropertyModifiedGetters::GetHasPixelAnimationModified },
+	{"PID_EnableTessellation", "Tessellation", EMaterialInstanceColumnType::StaticProperty, &GenerateTextForStaticProperty <&FMaterialInstanceStaticPropertyGetters::GetEnableTessellation>, &FMaterialInstanceStaticPropertyModifiedGetters::GetEnableTessellationModified },
+	{"PID_BlendMode", "Blend Mode", EMaterialInstanceColumnType::StaticProperty, &GenerateTextForStaticProperty <&FMaterialInstanceStaticPropertyGetters::GetBlendMode>, &FMaterialInstanceStaticPropertyModifiedGetters::GetBlendModeModified },
+	{"PID_ShadingModel", "Shading Model", EMaterialInstanceColumnType::StaticProperty, &GenerateTextForStaticProperty <&FMaterialInstanceStaticPropertyGetters::GetShadingModel>, &FMaterialInstanceStaticPropertyModifiedGetters::GetShadingModelModified },
+	{"PID_OpacityMaskClipValue", "Opacity Mask Clip Value", EMaterialInstanceColumnType::StaticProperty, &GenerateTextForStaticProperty <&FMaterialInstanceStaticPropertyGetters::GetOpacityMaskClipValue>, &FMaterialInstanceStaticPropertyModifiedGetters::GetOpacityMaskClipModified },
+};
+
+/** Static array of column FNames for all material usages. */
+static const FName UsageFlagPropertyIds[MATUSAGE_MAX] = {
+	"PID_UsageSkeletalMesh",
+	"PID_UsageParticleSprites",
+	"PID_UsageBeamTrails",
+	"PID_UsageMeshParticles",
+	"PID_UsageStaticLighting",
+	"PID_UsageMorphTargets",
+	"PID_UsageSplineMesh",
+	"PID_UsageInstancedStaticMeshes",
+	"PID_UsageGeometryCollections",
+	"PID_UsageClothing",
+	"PID_UsageNiagaraSprites",
+	"PID_UsageNiagaraRibbons",
+	"PID_UsageNiagaraMeshParticles",
+	"PID_UsageGeometryCache",
+	"PID_UsageWater",
+	"PID_UsageHairStrands",
+	"PID_UsageLidarPointCloud",
+	"PID_UsageVirtualHeightfieldMesh",
+	"PID_UsageNanite",
+	"PID_UsageVoxels",
+	"PID_UsageVolumetricCloud",
+	"PID_UsageHeterogeneousVolumes",
+	"PID_UsageStaticMesh",
+	"PID_UsageEditorCompositing",
+	"PID_UsageNeuralNetworks",
+	"PID_UsageMeshDeformer",
+};
+
+static int32 GetPropertyCount(FMaterialDatabaseAssetHierarchyInfo const& InInfo)
+{
+	return UE_ARRAY_COUNT(MaterialValidationLibrary::StaticPropertyInfos) + MATUSAGE_MAX + InInfo.StaticSwitchNames.Num() + InInfo.ComponentMaskNames.Num();
+}
+
+} // MaterialValidationLibrary
+
+void UMaterialValidationLibrary::GetMaterialProperties(FMaterialDatabaseAssetHierarchyInfo const& InInfo, TArray<FMaterialDatabaseAssetPropertyDesc>& OutProperties)
+{
+	const int32 PropertyCount = MaterialValidationLibrary::GetPropertyCount(InInfo);
+	OutProperties.Reset(PropertyCount);
+
+	for (MaterialValidationLibrary::FMaterialInstanceStaticPropertyInfo const& StaticPropertyInfo : MaterialValidationLibrary::StaticPropertyInfos)
+	{
+		OutProperties.Add({ .Id = StaticPropertyInfo.PropertyId, .Name = StaticPropertyInfo.PropertyName, .Category = StaticPropertyInfo.Category });
+	}
+	for (int32 Index = 0; Index < MATUSAGE_MAX; ++Index)
+	{
+		OutProperties.Add({ .Id = MaterialValidationLibrary::UsageFlagPropertyIds[Index], .Name = FName(*UMaterialInterface::GetUsageName((EMaterialUsage)Index)), .Category = EMaterialInstanceColumnType::UsageFlag});
+	}
+	for (const FName Name : InInfo.StaticSwitchNames)
+	{
+		OutProperties.Add({ .Id = Name, .Name = Name, .Category = EMaterialInstanceColumnType::StaticSwitch});
+	}
+	for (const FName Name : InInfo.ComponentMaskNames)
+	{
+		OutProperties.Add({ .Id = Name, .Name = Name, .Category = EMaterialInstanceColumnType::ComponentMask });
+	}
+}
+	
+void UMaterialValidationLibrary::GetMaterialPropertyValues(FMaterialDatabaseAssetHierarchyInfo const& InInfo, int32 InMaterialIndex, TArray<FMaterialDatabaseAssetPropertyValue>& OutValues)
+{
+	const int32 PropertyCount = MaterialValidationLibrary::GetPropertyCount(InInfo);
+	OutValues.Reset(PropertyCount);
+
+	if (!ensure(InInfo.MaterialAssetDatas.IsValidIndex(InMaterialIndex)))
+	{
+		return;
+	}
+
+	FMaterialValidationAssetData_MaterialInstance const& AssetData = InInfo.MaterialAssetDatas[InMaterialIndex];
+
+	for (MaterialValidationLibrary::FMaterialInstanceStaticPropertyInfo const& StaticPropertyInfo : MaterialValidationLibrary::StaticPropertyInfos)
+	{
+		OutValues.Add({
+			.Value = StaticPropertyInfo.GetPropertyValueTextFunction(AssetData), 
+			.bIsModifiedFromParent = StaticPropertyInfo.GetIsPropertyModifiedFunction(AssetData) });
+	}
+	
+	for (int32 Index = 0; Index < MATUSAGE_MAX; ++Index)
+	{
+		const bool bValue = (AssetData.StaticProperties.UsageFlags & (1u << Index)) != 0;
+		const bool bModified = (AssetData.StaticPropertyOverrideFlags.Override_UsageFlags & (1u << Index)) != 0;
+		OutValues.Add({	.Value = FText::AsNumber(bValue ? 1 : 0), .bIsModifiedFromParent = bModified });
+	}
+
+	for (int32 Index = 0; Index < InInfo.StaticSwitchNames.Num(); ++Index)
+	{
+		const uint32 BitArrayIndex = (uint32)Index / 32;
+		const uint32 BitIndex = (uint32)Index % 32;
+		
+		if (ensure(AssetData.StaticSwitchOverrideValues.IsValidIndex(BitArrayIndex) && AssetData.StaticSwitchOverrideMask.IsValidIndex(BitArrayIndex)))
+		{
+			const bool bValue = (AssetData.StaticSwitchOverrideValues[BitArrayIndex] & (1u << BitIndex)) != 0;
+			const bool bModified = (AssetData.StaticSwitchOverrideMask[BitArrayIndex] & (1u << BitIndex)) != 0;
+			OutValues.Add({ .Value = FText::AsNumber(bValue ? 1 : 0), .bIsModifiedFromParent = bModified });
+		}
+		else
+		{
+			// Mismatch possibly caused by add/removal of parameters since database build.
+			OutValues.Add({ .Value = FText::FromString("?"), .bIsModifiedFromParent = false });
+		}
+	}
+
+	for (int32 Index = 0; Index < InInfo.ComponentMaskNames.Num(); ++Index)
+	{
+		const uint32 BitArrayIndex = Index / 8;
+		const uint32 BitIndex = (Index * 4) % 32;
+
+		if (ensure(AssetData.ComponentMaskOverrideValues.IsValidIndex(BitArrayIndex) && AssetData.ComponentMaskOverrideMask.IsValidIndex(BitArrayIndex)))
+		{
+			const uint32 Value = (AssetData.ComponentMaskOverrideValues[BitArrayIndex] >> BitIndex) & 0xf;
+			const bool bModified = ((AssetData.ComponentMaskOverrideMask[BitArrayIndex] >> BitIndex) & 0xf) != 0;
+			TStringBuilder<8> RGBA;
+			RGBA.AppendChar(Value & 1 ? 'R' : '_');
+			RGBA.AppendChar(Value & 2 ? 'G' : '_');
+			RGBA.AppendChar(Value & 4 ? 'B' : '_');
+			RGBA.AppendChar(Value & 8 ? 'A' : '_');
+			OutValues.Add({ .Value = FText::FromString(*RGBA), .bIsModifiedFromParent = bModified });
+		}
+		else
+		{
+			// Mismatch possibly caused by add/removal of parameters since database build.
+			OutValues.Add({ .Value = FText::FromString("?"), .bIsModifiedFromParent = false });
+		}
 	}
 }
 
diff --git a/Engine/Plugins/Editor/MaterialValidation/Source/Private/MaterialValidationLibrary.h b/Engine/Plugins/Editor/MaterialValidation/Source/Private/MaterialValidationLibrary.h
--- a/Engine/Plugins/Editor/MaterialValidation/Source/Private/MaterialValidationLibrary.h
+++ b/Engine/Plugins/Editor/MaterialValidation/Source/Private/MaterialValidationLibrary.h
@@ -6,10 +6,77 @@
 #include "Kismet/BlueprintFunctionLibrary.h"
 #include "MaterialValidationLibrary.generated.h"
 
+struct FMaterialValidationAssetData_MaterialInstance;
 class UMaterial;
 class UMaterialInterface;
 class UMaterialInstanceConstant;
 
+/** Enumeration of material instance property types. */
+UENUM(BlueprintType)
+enum class EMaterialInstanceColumnType : uint8
+{
+	Default,
+	StaticProperty,
+	UsageFlag,
+	StaticSwitch,
+	ComponentMask,
+	Count
+};
+
+/** 
+ * Structure that collects all the resolved asset data for one material hierarchy in the material asset database. 
+ * Typically this is extracted from the compact form stored in a UMaterialValidationGroup and used in UI views or for further analysis.
+ */
+USTRUCT(BlueprintType)
+struct FMaterialDatabaseAssetHierarchyInfo
+{
+	GENERATED_BODY()
+
+	/** Array of material asset paths. This includes the base UMaterial asset in the first entry, and the child UMaterialInstance assets after that. */
+	UPROPERTY(BlueprintReadOnly, Category = Material)
+	TArray<FString> MaterialPaths;
+	/** Array of asset data info. This is synced with the MaterialPaths array. */
+	UPROPERTY()
+	TArray<FMaterialValidationAssetData_MaterialInstance> MaterialAssetDatas;
+	/** Array of all StaticSwitch names. This matches the layout of all StaticSwitchOverrideValues arrays in the MaterialAssetDatas. */
+	UPROPERTY()
+	TArray<FName> StaticSwitchNames;
+	/** Array of all ComponentMask names. This matches the layout of all ComponentMaskOverrideValues arrays in the MaterialAssetDatas. */
+	UPROPERTY()
+	TArray<FName> ComponentMaskNames;
+};
+
+/** Description of a material property in the material asset database. */
+USTRUCT(BlueprintType)
+struct FMaterialDatabaseAssetPropertyDesc
+{
+	GENERATED_BODY()
+
+	/** Id for the property. */
+	UPROPERTY(BlueprintReadOnly, Category = Material)
+	FName Id;
+	/** Readable name for the property. */
+	UPROPERTY(BlueprintReadOnly, Category = Material)
+	FName Name;
+	/** Category for the property. */
+	UPROPERTY(BlueprintReadOnly, Category = Material)
+	EMaterialInstanceColumnType Category = EMaterialInstanceColumnType::Default;
+};
+
+/** Description of a resolved material value in the material asset database. */
+USTRUCT(BlueprintType)
+struct FMaterialDatabaseAssetPropertyValue
+{
+	GENERATED_BODY()
+
+	/** Text value for the property. */
+	UPROPERTY(BlueprintReadOnly, Category = Material)
+	FText Value;
+	/** true if the value for the property is different from the value on its parent material. */
+	UPROPERTY(BlueprintReadOnly, Category = Material)
+	bool bIsModifiedFromParent = false;
+};
+
 /** Material validation functionality exposed as a library of Blueprint functions. */
 UCLASS(MinimalAPI, meta = (ScriptName = "MaterialValidationLibrary"))
 class UMaterialValidationLibrary : public UBlueprintFunctionLibrary
@@ -104,12 +171,32 @@ class UMaterialValidationLibrary : public UBlueprintFunctionLibrary
 	UFUNCTION(BlueprintCallable, Category = Material)
 	static int32 GetModifiedShaderCount(UMaterialValidationGroup const* InGroup, UMaterial* InMaterial, TArray<UMaterialInterface*>& InModifiedObjects);
 
-	/** Get all of the material instance data for a single material. */
-	static void GetMaterialInstanceAssetDatas(
-		UMaterialValidationGroup const* InGroup,
-		FString const& MaterialPath,
-		TArray<FString>& OutMaterialInstancePaths,
-		TArray<struct FMaterialValidationAssetData_MaterialInstance>& OutMaterialInstanceAssetDatas,
-		TArray<FName>& OutStaticSwitchNames,
-		TArray<FName>& OutComponentMaskNames);
+	/** 
+	 * Get all of the material instance data for a single material hierarchy. 
+	 *
+	 * @param	InGroup				The group to use.
+	 * @param	InBaseMaterialPath	The path of the base material to evaluate.
+	 * @param	OutInfo				Returnsa an object containing all of the material instance data.
+	 */
+	UFUNCTION(BlueprintCallable, Category = Material)
+	static void GetMaterialHierarchyInfo(UMaterialValidationGroup const* InGroup, FString const& InBaseMaterialPath, FMaterialDatabaseAssetHierarchyInfo& OutInfo);
+
+	/**
+	 * Get all the material properties for a material hierarchy.
+	 *
+	 * @param	InInfo			The full hierarchy info. This is typically resolved by first calling GetMaterialHierarchyInfo().
+	 * @param	OutProperties	Returns an array of property descriptions.
+	 */
+	UFUNCTION(BlueprintCallable, Category = Material)
+	static void GetMaterialProperties(FMaterialDatabaseAssetHierarchyInfo const& InInfo, TArray<FMaterialDatabaseAssetPropertyDesc>& OutProperties);
+	
+	/**
+	 * Get all the material property values for a material instance in a material hierarchy
+	 *
+	 * @param	InInfo			The full hierarchy info. This is typically resolved by first calling GetMaterialHierarchyInfo().
+	 * @param	InMaterialIndex	The index of the material instance in the hierarchy. The instances are in the order defined by FMaterialDatabaseAssetHierarchyInfo::MaterialPaths.
+	 * @param	OutValues		Returns an array of property values, one for each of the material properties as returned by GetProperties().
+	 */
+	UFUNCTION(BlueprintCallable, Category = Material)
+	static void GetMaterialPropertyValues(FMaterialDatabaseAssetHierarchyInfo const& InInfo, int32 InMaterialIndex, TArray<FMaterialDatabaseAssetPropertyValue>& OutValues);
 };
diff --git a/Engine/Plugins/Editor/MaterialValidation/Source/Private/MaterialValidationToolkit.cpp b/Engine/Plugins/Editor/MaterialValidation/Source/Private/MaterialValidationToolkit.cpp
--- a/Engine/Plugins/Editor/MaterialValidation/Source/Private/MaterialValidationToolkit.cpp
+++ b/Engine/Plugins/Editor/MaterialValidation/Source/Private/MaterialValidationToolkit.cpp
@@ -1219,25 +1219,16 @@ static TSharedPtr<SMaterialInstanceTable> CreateMaterialInstanceTable()
 /** Update the table for all material instances in a material. */
 static void UpdateMaterialInstanceTable(TSharedPtr<SMaterialInstanceTable> InTable, UMaterialValidationGroup const* InGroup, FString const& InMaterialPath)
 {
-	TArray<FString> MaterialInstancePaths;
-	TArray<FMaterialValidationAssetData_MaterialInstance> MaterialInstanceAssetDatas;
-	TSharedPtr<FMaterialInstanceParameterNames> ParameterNames = MakeShared<FMaterialInstanceParameterNames>();
-	
 	// Get all the information from the material validation database.
-	UMaterialValidationLibrary::GetMaterialInstanceAssetDatas(
-		InGroup, 
-		InMaterialPath, 
-		MaterialInstancePaths, 
-		MaterialInstanceAssetDatas, 
-		ParameterNames->StaticSwitches, 
-		ParameterNames->ComponentMasks);
+	FMaterialDatabaseAssetHierarchyInfo AssetInfo;
+	UMaterialValidationLibrary::GetMaterialHierarchyInfo(InGroup, InMaterialPath, AssetInfo);
 
 	// Build all rows.
 	TArray<TSharedPtr<FMaterialInstanceRowData>> Rows;
-	const int32 NumMaterialInstances = MaterialInstancePaths.Num();
+	const int32 NumMaterialInstances = AssetInfo.MaterialPaths.Num();
 	for (int32 Index = 0; Index < NumMaterialInstances; ++Index)
 	{
-		Rows.Add(MakeShared<FMaterialInstanceRowData>(MaterialInstancePaths[Index], MaterialInstanceAssetDatas[Index]));
+		Rows.Add(MakeShared<FMaterialInstanceRowData>(AssetInfo.MaterialPaths[Index], AssetInfo.MaterialAssetDatas[Index]));
 	}
 
 	// Build all columns.
@@ -1265,17 +1256,22 @@ static void UpdateMaterialInstanceTable(TSharedPtr<SMaterialInstanceTable> InTab
 		{
 			Columns.Add(MakeShared<FMaterialInstanceColumnData>(UsageColumnIds[UsageIndex], FText::FromString(UMaterialInterface::GetUsageName((EMaterialUsage)UsageIndex)), EMaterialInstanceColumnType::UsageFlag));
 		}
-		for (FName& Name : ParameterNames->StaticSwitches)
+		for (FName& Name : AssetInfo.StaticSwitchNames)
 		{
 			Columns.Add(MakeShared<FMaterialInstanceColumnData>(Name, FText::FromName(Name), EMaterialInstanceColumnType::StaticSwitch));
 		}
-		for (FName& Name : ParameterNames->ComponentMasks)
+		for (FName& Name : AssetInfo.ComponentMaskNames)
 		{
 			Columns.Add(MakeShared<FMaterialInstanceColumnData>(Name, FText::FromName(Name), EMaterialInstanceColumnType::ComponentMask));
 		}
 	}
 	
-	// Build FMaterialInstanceParameterNames lookups.
+	// Move parameter names and build FMaterialInstanceParameterNames lookups.
+	TSharedPtr<FMaterialInstanceParameterNames> ParameterNames = MakeShared<FMaterialInstanceParameterNames>();
+
+	ParameterNames->StaticSwitches = MoveTemp(AssetInfo.StaticSwitchNames);
+	ParameterNames->ComponentMasks = MoveTemp(AssetInfo.ComponentMaskNames);
+
 	for (int32 Index = 0; Index < ParameterNames->StaticSwitches.Num(); ++Index)
 	{
 		ParameterNames->StaticSwitchNameToIndex.Add(ParameterNames->StaticSwitches[Index], Index);
diff --git a/Engine/Plugins/Editor/MaterialValidation/Source/Private/MaterialValidationToolkit.h b/Engine/Plugins/Editor/MaterialValidation/Source/Private/MaterialValidationToolkit.h
--- a/Engine/Plugins/Editor/MaterialValidation/Source/Private/MaterialValidationToolkit.h
+++ b/Engine/Plugins/Editor/MaterialValidation/Source/Private/MaterialValidationToolkit.h
@@ -17,15 +17,3 @@ class UAssetDefinition_MaterialValidationGroup : public UAssetDefinitionDefault
 	virtual TConstArrayView<FAssetCategoryPath> GetAssetCategories() const override;
 	virtual EAssetCommandResult OpenAssets(const FAssetOpenArgs& OpenArgs) const override;
 };
-
-/** Enumeration of material instance property types used to filter editor toolkit tables. */
-UENUM()
-enum class EMaterialInstanceColumnType
-{
-	Default,
-	StaticProperty,
-	UsageFlag,
-	StaticSwitch,
-	ComponentMask,
-	Count
-};
