MaterialValidationLibrary: Return FSoftObjectPath instead of path FString for exposed Blueprint APIs.
Expose UMaterialValidationGroup as BlueprintType to get the array of material FSoftObjectPaths.
Multiple fixes to handling of redirectors in MaterialValidationLibrary functions.
Redirectors aren't automatically handled because we store FString paths in the asset to avoid asset refs (which tie the database asset to material rename/delete workflows).
Some handling was in place, but reverse lookup of soft object path to string paths used for map keys was missing.
#rnx
#Synced-CL 51290350

diff --git a/Engine/Plugins/Editor/MaterialValidation/Source/Private/MaterialValidationGroup.cpp b/Engine/Plugins/Editor/MaterialValidation/Source/Private/MaterialValidationGroup.cpp
--- a/Engine/Plugins/Editor/MaterialValidation/Source/Private/MaterialValidationGroup.cpp
+++ b/Engine/Plugins/Editor/MaterialValidation/Source/Private/MaterialValidationGroup.cpp
@@ -127,6 +127,15 @@ void UMaterialValidationGroup::PostEditChangeProperty(FPropertyChangedEvent& Pro
 	}
 }
 
+void UMaterialValidationGroup::GetMaterialPaths(TArray<FSoftObjectPath>& OutMaterialPaths) const
+{
+	OutMaterialPaths.Empty(Materials.Num());
+	for (TPair<FString, FMaterialValidationDesc> const& It : Materials)
+	{
+		OutMaterialPaths.Add(UMaterialValidationLibrary::ResolveAssetPath(It.Key));
+	}
+}
+
 void UMaterialValidationGroup::UpdateMaterials()
 {
 	UMaterialValidationLibrary::RemoveInvalidMaterialsFromGroup(this);
diff --git a/Engine/Plugins/Editor/MaterialValidation/Source/Private/MaterialValidationGroup.h b/Engine/Plugins/Editor/MaterialValidation/Source/Private/MaterialValidationGroup.h
--- a/Engine/Plugins/Editor/MaterialValidation/Source/Private/MaterialValidationGroup.h
+++ b/Engine/Plugins/Editor/MaterialValidation/Source/Private/MaterialValidationGroup.h
@@ -41,7 +41,7 @@ struct FMaterialValidationDesc
 };
 
 /** Data Asset containing the validation setup for a group of UMaterials. */
-UCLASS(MinimalAPI)
+UCLASS(MinimalAPI, BlueprintType)
 class UMaterialValidationGroup : public UDataAsset
 {
 	GENERATED_BODY()
@@ -67,6 +67,10 @@ public:
 	UPROPERTY()
 	TMap<FString, FMaterialValidationDesc> Materials;
 
+	/** Get UMaterial assets resolved to FSoftObjectPaths. */
+	UFUNCTION(BlueprintPure, Category = Material)
+	void GetMaterialPaths(TArray<FSoftObjectPath>& OutMaterialPaths) const;
+
 	/** 
 	 * Add all Material assets found under the material search paths.
 	 * The new material assets will have empty permutation lists.
diff --git a/Engine/Plugins/Editor/MaterialValidation/Source/Private/MaterialValidationLibrary.cpp b/Engine/Plugins/Editor/MaterialValidation/Source/Private/MaterialValidationLibrary.cpp
--- a/Engine/Plugins/Editor/MaterialValidation/Source/Private/MaterialValidationLibrary.cpp
+++ b/Engine/Plugins/Editor/MaterialValidation/Source/Private/MaterialValidationLibrary.cpp
@@ -487,6 +487,51 @@ static FSoftObjectPath ConvertStringPathToSoftObjectPath(FString const& InAssetP
 	return CurrentPath;
 }
 
+/** Find the material description in a group that matches a FSoftObjectPath taking into account any asset redirectors. */
+static FMaterialValidationDesc const* FindMaterialDescFromSoftObjectPath(UMaterialValidationGroup const& InGroup, FSoftObjectPath const& InSoftObjectPath)
+{
+	// First look at for simple converted value of the soft object path.
+	FString ConvertedString = InSoftObjectPath.ToString();
+	if (FMaterialValidationDesc const* MaterialDesc = InGroup.Materials.Find(ConvertedString))
+	{
+		return MaterialDesc;
+	}
+
+	// If not found then a redirector might be in play so do a slow find over all stored material paths.
+	for (TPair<FString, FMaterialValidationDesc> const& It : InGroup.Materials)
+	{
+		if (MaterialValidationLibrary::ConvertStringPathToSoftObjectPath(It.Key) == InSoftObjectPath)
+		{
+			return & It.Value;
+		}
+	}
+
+	return nullptr;
+}
+
+static bool FindMaterialInstanceIndexFromSoftObjectPath(FMaterialValidationDesc const& InMaterialDesc, FSoftObjectPath const& InSoftObjectPath, int32& OutIndex)
+{
+	// First look at for simple converted value of the soft object path.
+	FString ConvertedString = InSoftObjectPath.ToString();
+	if (int32 const* IndexPtr = InMaterialDesc.MaterialInstances.Find(ConvertedString))
+	{
+		OutIndex = *IndexPtr;
+		return true;
+	}
+
+	// If not found then a redirector might be in play so do a slow find over all stored material instance paths.
+	for (TPair<FString, int32> const& It : InMaterialDesc.MaterialInstances)
+	{
+		if (MaterialValidationLibrary::ConvertStringPathToSoftObjectPath(It.Key) == InSoftObjectPath)
+		{
+			OutIndex = It.Value;
+			return true;
+		}
+	}
+
+	return false;
+}
+
 /** 
  * Add the package path extension to a FName package path that has a missing extension. 
  * todo: Ideally we would avoid this by only passing around package paths in one canonical form.
@@ -537,14 +582,28 @@ static void GetMaterialInstanceHierarchy(UMaterial const& InMaterial, TConstArra
 	// Maintain a stack of depths so that we can output those to OutHierarchyDepths.
 	TArray<int32> DepthStack;
 	DepthStack.Push(0);
+	// Maintain a stack of redirector flags. We follow the redirectors but don't add them to the output.
+	TArray<bool> RedirectorStack;
+	RedirectorStack.Push(false);
+	// Maintain an array of names we have seen but discarded to avoid cyclic iteration.
+	TArray<FName> DiscardedPackageNames;
 
 	while (!NameStack.IsEmpty())
 	{
 		// Pop from stack and add to output.
 		const FName ParentPackageName = NameStack.Pop(EAllowShrinking::No);
-		OutPackageNames.Add(ParentPackageName);
 		const int32 ParentPackageDepth = DepthStack.Pop(EAllowShrinking::No);
-		OutHierarchyDepths.Add(ParentPackageDepth);
+		const bool bParentIsRedirector = RedirectorStack.Pop(EAllowShrinking::No);
+
+		if (bParentIsRedirector)
+		{
+			DiscardedPackageNames.Add(ParentPackageName);
+		}
+		else
+		{
+			OutPackageNames.Add(ParentPackageName);
+			OutHierarchyDepths.Add(ParentPackageDepth);
+		}
 
 		// Push children onto our stack.
 		TArray<FName> ChildPackageNames;
@@ -552,7 +611,7 @@ static void GetMaterialInstanceHierarchy(UMaterial const& InMaterial, TConstArra
 
 		for (const FName& PackageName : ChildPackageNames)
 		{
-			if (OutPackageNames.Contains(PackageName) || NameStack.Contains(PackageName) || IsPackageInPaths(PackageName.ToString(), InExcludePaths))
+			if (OutPackageNames.Contains(PackageName) || NameStack.Contains(PackageName) || DiscardedPackageNames.Contains(PackageName) ||  IsPackageInPaths(PackageName.ToString(), InExcludePaths))
 			{
 				continue;
 			}
@@ -566,14 +625,22 @@ static void GetMaterialInstanceHierarchy(UMaterial const& InMaterial, TConstArra
 				{
 					continue;
 				}
-				if (!AssetData.IsInstanceOf(UMaterialInstanceConstant::StaticClass()))
+
+				if (AssetData.IsRedirector())
 				{
-					continue;
+					NameStack.Push(PackageName);
+					DepthStack.Push(ParentPackageDepth);
+					RedirectorStack.Push(true);
+					break;
 				}
 
-				NameStack.Push(PackageName);
-				DepthStack.Push(ParentPackageDepth + 1);
-				break;
+				if (AssetData.IsInstanceOf(UMaterialInstanceConstant::StaticClass()))
+				{
+					NameStack.Push(PackageName);
+					DepthStack.Push(ParentPackageDepth + 1);
+					RedirectorStack.Push(false);
+					break;
+				}
 			}
 		}
 	}
@@ -719,38 +786,34 @@ static void UpdateValidMaterialInstances(UMaterial const& InMaterial, TConstArra
 /** 
  * Get the asset data associated with a package path.
  * We get this from either the material validation group or from the object directly.
- * We only set OutMaterialInstance if the object is loaded.
+ * We only set InOutMaterialInstance if the object is loaded.
  * For performance we avoid loading an object if it is unloaded (but could optionally do that if that becomes useful).
  * If we return false then it is likely that the material instance exists but:
- * - It wasn't stored in the validation group (it was probably created after last validation group update)
- * - It isn't in the InObjectPackagePaths (this usually means that it isn't in the curent validation changelist)
+ * - It wasn't stored in the validation group (it was probably created after last validation group update), and
+ * - It isn't in the curent validation changelist, and
  * - It is not currently loaded.
  */
 static bool GetMaterialInstanceAssetData(
 	FString const& InPackagePath,
-	TMap<FString, UMaterialInterface*> const& InObjectPackagePaths,
+	UMaterialInstanceConstant*& InOutMaterialInstance,
 	FMaterialValidationDesc const* InMaterialDesc,
 	FMaterialPermutationData const& InMaterialPermutationData,
-	FMaterialValidationAssetData_MaterialInstance& OutAssetData,
-	UMaterialInstanceConstant*& OutMaterialInstance)
+	FMaterialValidationAssetData_MaterialInstance& OutAssetData)
 {
-	// If we find in our known override object paths then initialize the asset data from the object.
-	if (UMaterialInterface* const* ModifiedAsset = InObjectPackagePaths.Find(InPackagePath))
+	// If we were provided with a modified object then use that.
+	if (InOutMaterialInstance != nullptr)
 	{
-		if (UMaterialInstanceConstant* MaterialInstance = Cast<UMaterialInstanceConstant>(*ModifiedAsset))
-		{
-			GetAssetData(*MaterialInstance, InMaterialPermutationData, OutAssetData);
-			OutMaterialInstance = MaterialInstance;
-			return true;
-		}
+		GetAssetData(*InOutMaterialInstance, InMaterialPermutationData, OutAssetData);
+		return true;
 	}
 
 	// Otherwise use asset data from the material validation group.
 	if (InMaterialDesc != nullptr)
 	{
-		if (int32 const* AssetDataIndex = InMaterialDesc->MaterialInstances.Find(InPackagePath))
+		int32 AssetDataIndex = INDEX_NONE;
+		if (MaterialValidationLibrary::FindMaterialInstanceIndexFromSoftObjectPath(*InMaterialDesc, InPackagePath, AssetDataIndex))
 		{
-			if (*AssetDataIndex == INDEX_NONE)
+			if (AssetDataIndex == INDEX_NONE)
 			{
 				// Set up a default asset data with the right sizes.
 				const int32 StaticSwitchArrayNum = InMaterialDesc->AssetData.StaticSwitchValues.Num();
@@ -761,21 +824,21 @@ static bool GetMaterialInstanceAssetData(
 				OutAssetData.ComponentMaskOverrideValues.SetNumZeroed(ComponentMaskArrayNum);
 				return true;
 			}
-			else if (InMaterialDesc->MaterialInstanceAssetDatas.IsValidIndex(*AssetDataIndex))
+			else if (InMaterialDesc->MaterialInstanceAssetDatas.IsValidIndex(AssetDataIndex))
 			{
-				OutAssetData = InMaterialDesc->MaterialInstanceAssetDatas[*AssetDataIndex];
+				OutAssetData = InMaterialDesc->MaterialInstanceAssetDatas[AssetDataIndex];
 				return true;
 			}
 		}
 	}
 
 	// If we didn't find in the material validation group but the object is already loaded then use that.
-	FSoftObjectPath SoftObjectPath(InPackagePath);
+	FSoftObjectPath SoftObjectPath = MaterialValidationLibrary::ConvertStringPathToSoftObjectPath(InPackagePath);
 	TSoftObjectPtr<UMaterialInstanceConstant> SoftObjectPtr(SoftObjectPath);
 	if (UMaterialInstanceConstant* MaterialInstance = SoftObjectPtr.Get())
 	{
 		GetAssetData(*MaterialInstance, InMaterialPermutationData, OutAssetData);
-		OutMaterialInstance = MaterialInstance;
+		InOutMaterialInstance = MaterialInstance;
 		return true;
 	}
 
@@ -1016,6 +1079,11 @@ void UMaterialValidationLibrary::UpdateMaterialPermutationsInGroup(UMaterialVali
 	InGroup->Modify();
 }
 
+FSoftObjectPath UMaterialValidationLibrary::ResolveAssetPath(FString const& InAssetPath)
+{
+	return MaterialValidationLibrary::ConvertStringPathToSoftObjectPath(InAssetPath);
+}
+
 void UMaterialValidationLibrary::IsMaterialInGroup(UMaterialValidationGroup const* InGroup, UMaterial* InMaterial, bool& bOutIsInGroupPath, bool& bOutIsInGroup)
 {
 	bOutIsInGroupPath = false;
@@ -1027,9 +1095,7 @@ void UMaterialValidationLibrary::IsMaterialInGroup(UMaterialValidationGroup cons
 	}
 
 	const FSoftObjectPath SoftPackagePath(InMaterial);
-	const FString PackagePath = SoftPackagePath.ToString();
-
-	if (InGroup->Materials.Contains(PackagePath))
+	if (MaterialValidationLibrary::FindMaterialDescFromSoftObjectPath(*InGroup, SoftPackagePath) != nullptr)
 	{
 		bOutIsInGroup = true;
 	}
@@ -1038,6 +1104,7 @@ void UMaterialValidationLibrary::IsMaterialInGroup(UMaterialValidationGroup cons
 	TArray<FString> ExcludePaths;
 	MaterialValidationLibrary::GetSearchPaths(*InGroup, SearchPaths, ExcludePaths);
 
+	const FString PackagePath = SoftPackagePath.ToString();
 	if (MaterialValidationLibrary::IsPackageInPaths(PackagePath, SearchPaths))
 	{
 		bOutIsInGroupPath = true;
@@ -1061,9 +1128,8 @@ void UMaterialValidationLibrary::IsMaterialInstanceInGroup(UMaterialValidationGr
 
 	UMaterial const* Material = InMaterialInstance->GetMaterial();
 	const FSoftObjectPath SoftPackagePath(Material);
-	const FString PackagePath = SoftPackagePath.ToString();
 
-	FMaterialValidationDesc const* MaterialDesc = InGroup->Materials.Find(PackagePath);
+	FMaterialValidationDesc const* MaterialDesc = MaterialValidationLibrary::FindMaterialDescFromSoftObjectPath(*InGroup, SoftPackagePath);
 	if (MaterialDesc == nullptr)
 	{
 		return;
@@ -1083,9 +1149,8 @@ int32 UMaterialValidationLibrary::GetShaderCount(UMaterialValidationGroup const*
 	}
 
 	const FSoftObjectPath SoftPackagePath(InBaseMaterial);
-	const FString PackagePath = SoftPackagePath.ToString();
 
-	FMaterialValidationDesc const* MaterialDesc = InGroup->Materials.Find(PackagePath);
+	FMaterialValidationDesc const* MaterialDesc = MaterialValidationLibrary::FindMaterialDescFromSoftObjectPath(*InGroup, SoftPackagePath);
 	return MaterialDesc != nullptr ? MaterialDesc->NumShadersTotal : 0;
 }
 
@@ -1097,11 +1162,10 @@ int32 UMaterialValidationLibrary::GetModifiedShaderCount(UMaterialValidationGrou
 	}
 
 	// Get paths to modified objects for lookup.
-	TMap<FString, UMaterialInterface*> ModifiedObjectPackagePaths;
+	TMap<FSoftObjectPath, UMaterialInterface*> ModifiedObjectPaths;
 	for (UMaterialInterface* Object : InModifiedObjects)
 	{
-		const FSoftObjectPath SoftPackagePath(Object);
-		ModifiedObjectPackagePaths.Add(SoftPackagePath.ToString(), Object);
+		ModifiedObjectPaths.Add(FSoftObjectPath(Object), Object);
 	}
 
 	TArray<FString> SearchPaths;
@@ -1117,11 +1181,10 @@ int32 UMaterialValidationLibrary::GetModifiedShaderCount(UMaterialValidationGrou
 	// Get the material asset data stored on the group.
 	// But evaluate from the material if the material is modified or if it is not found in the group.
 	const FSoftObjectPath  BaseMaterialSoftObjectPath(InBaseMaterial);
-	const FString BaseMaterialPackagePath = BaseMaterialSoftObjectPath.ToString();
-	FMaterialValidationDesc const* MaterialDesc = InGroup->Materials.Find(BaseMaterialPackagePath);
+	FMaterialValidationDesc const* MaterialDesc = MaterialValidationLibrary::FindMaterialDescFromSoftObjectPath(*InGroup, BaseMaterialSoftObjectPath);
 
 	FMaterialValidationAssetData_Material MaterialAssetData;
-	if (MaterialDesc == nullptr || ModifiedObjectPackagePaths.Contains(BaseMaterialPackagePath))
+	if (MaterialDesc == nullptr || InModifiedObjects.Contains(InBaseMaterial))
 	{
 		MaterialAssetData = MaterialAssetDataFromObject;
 	}
@@ -1167,10 +1230,12 @@ int32 UMaterialValidationLibrary::GetModifiedShaderCount(UMaterialValidationGrou
 	{
 		const int32 CurrentDepth = HierarchyDepths[Index];
 		FString PackagePath = MaterialValidationLibrary::AddPackagePathExtension(PackageNames[Index]);
+		FSoftObjectPath SoftObjectPackagePath = MaterialValidationLibrary::ConvertStringPathToSoftObjectPath(PackagePath);
+		UMaterialInterface** FoundModifiedObject = ModifiedObjectPaths.Find(SoftObjectPackagePath);
+		UMaterialInstanceConstant* MaterialInstance = FoundModifiedObject ? Cast<UMaterialInstanceConstant>(*FoundModifiedObject) : nullptr;
 
 		FMaterialValidationAssetData_MaterialInstance AssetData;
-		UMaterialInstanceConstant* MaterialInstance = nullptr;
-		const bool bValidAssetData = MaterialValidationLibrary::GetMaterialInstanceAssetData(PackagePath, ModifiedObjectPackagePaths, MaterialDesc, PermutationData, AssetData, MaterialInstance);
+		const bool bValidAssetData = MaterialValidationLibrary::GetMaterialInstanceAssetData(PackagePath, MaterialInstance, MaterialDesc, PermutationData, AssetData);
 
 		if (!bValidAssetData)
 		{
@@ -1216,7 +1281,7 @@ int32 UMaterialValidationLibrary::GetModifiedShaderCount(UMaterialValidationGrou
  * Note there is quite a bit of duplication with GetModifiedShaderCount() in this function. 
  * It might be good to consolidate on a single material instance visitor that can do arbitrary work?
  */
-void UMaterialValidationLibrary::GetMaterialHierarchyInfo(UMaterialValidationGroup const* InGroup, FString const& InBaseMaterialPath, FMaterialDatabaseAssetHierarchyInfo& OutInfo)
+void UMaterialValidationLibrary::GetMaterialHierarchyInfo(UMaterialValidationGroup const* InGroup, FSoftObjectPath const& InBaseMaterialPath, FMaterialDatabaseAssetHierarchyInfo& OutInfo)
 {
 	OutInfo.MaterialPaths.Reset();
 	OutInfo.MaterialAssetDatas.Reset();
@@ -1228,13 +1293,13 @@ void UMaterialValidationLibrary::GetMaterialHierarchyInfo(UMaterialValidationGro
 		return;
 	}
 
-	FMaterialValidationDesc const* MaterialDesc = InGroup->Materials.Find(InBaseMaterialPath);
+	FMaterialValidationDesc const* MaterialDesc = MaterialValidationLibrary::FindMaterialDescFromSoftObjectPath(*InGroup, InBaseMaterialPath);
 	if (MaterialDesc == nullptr)
 	{
 		return;
 	}
 
-	TSoftObjectPtr<UMaterial> MaterialSoftPtr(MaterialValidationLibrary::ConvertStringPathToSoftObjectPath(InBaseMaterialPath));
+	TSoftObjectPtr<UMaterial> MaterialSoftPtr(InBaseMaterialPath);
 	UMaterial const* Material = MaterialSoftPtr.LoadSynchronous();
 	if (Material == nullptr)
 	{
@@ -1294,9 +1359,11 @@ void UMaterialValidationLibrary::GetMaterialHierarchyInfo(UMaterialValidationGro
 
 		bool bValidAssetData = false;
 		FMaterialValidationAssetData_MaterialInstance AssetData;
-		if (int32 const* AssetDataIndex = MaterialDesc->MaterialInstances.Find(PackagePath))
+		
+		int32 AssetDataIndex = INDEX_NONE;
+		if (MaterialValidationLibrary::FindMaterialInstanceIndexFromSoftObjectPath(*MaterialDesc, PackagePath, AssetDataIndex))
 		{
-			if (*AssetDataIndex == INDEX_NONE)
+			if (AssetDataIndex == INDEX_NONE)
 			{
 				// Set up a default asset data with the right sizes.
 				const int32 StaticSwitchArrayNum = MaterialDesc->AssetData.StaticSwitchValues.Num();
@@ -1307,9 +1374,9 @@ void UMaterialValidationLibrary::GetMaterialHierarchyInfo(UMaterialValidationGro
 				AssetData.ComponentMaskOverrideValues.SetNumZeroed(ComponentMaskArrayNum);
 				bValidAssetData = true;
 			}
-			else if (MaterialDesc->MaterialInstanceAssetDatas.IsValidIndex(*AssetDataIndex))
+			else if (MaterialDesc->MaterialInstanceAssetDatas.IsValidIndex(AssetDataIndex))
 			{
-				AssetData = MaterialDesc->MaterialInstanceAssetDatas[*AssetDataIndex];
+				AssetData = MaterialDesc->MaterialInstanceAssetDatas[AssetDataIndex];
 				bValidAssetData = true;
 			}
 		}
@@ -1337,7 +1404,7 @@ void UMaterialValidationLibrary::GetMaterialHierarchyInfo(UMaterialValidationGro
 		Child.MaterialLayerHash = AssetData.MaterialLayerHash;
 		Child.PermutationHash = MaterialValidationLibrary::BuildPermutationHash(Child.StaticProperties, Child.StaticSwitchOverrideValues, Child.ComponentMaskOverrideValues, Child.MaterialLayerHash);
 
-		OutInfo.MaterialPaths.Add(PackagePath);
+		OutInfo.MaterialPaths.Add(MaterialValidationLibrary::ConvertStringPathToSoftObjectPath(PackagePath));
 		OutInfo.MaterialAssetDatas.Add(Child);
 	}
 }
diff --git a/Engine/Plugins/Editor/MaterialValidation/Source/Private/MaterialValidationLibrary.h b/Engine/Plugins/Editor/MaterialValidation/Source/Private/MaterialValidationLibrary.h
--- a/Engine/Plugins/Editor/MaterialValidation/Source/Private/MaterialValidationLibrary.h
+++ b/Engine/Plugins/Editor/MaterialValidation/Source/Private/MaterialValidationLibrary.h
@@ -34,7 +34,7 @@ struct FMaterialDatabaseAssetHierarchyInfo
 
 	/** Array of material asset paths. This includes the base UMaterial asset in the first entry, and the child UMaterialInstance assets after that. */
 	UPROPERTY(BlueprintReadOnly, Category = Material)
-	TArray<FString> MaterialPaths;
+	TArray<FSoftObjectPath> MaterialPaths;
 	/** Array of asset data info. This is synced with the MaterialPaths array. */
 	UPROPERTY()
 	TArray<FMaterialValidationAssetData_MaterialInstance> MaterialAssetDatas;
@@ -126,6 +126,15 @@ class UMaterialValidationLibrary : public UBlueprintFunctionLibrary
 	UFUNCTION(BlueprintCallable, Category = Material)
 	static void UpdateMaterialPermutationsInGroup(UMaterialValidationGroup* InGroup);
 
+	/**
+	 * Resolve a asset path string to a FSoftObjectPath while handling asset redirectors.
+	 *
+	 * @param	InAssetPath		A asset path.
+	 *
+	 * @returns	The FSoftObjectPath of the resolved asset.
+	 */
+	static FSoftObjectPath ResolveAssetPath(FString const& InAssetPath);
+
 	/**
 	 * Get whether a material belongs to a group.
 	 * 
@@ -179,7 +188,7 @@ class UMaterialValidationLibrary : public UBlueprintFunctionLibrary
 	 * @param	OutInfo				Returnsa an object containing all of the material instance data.
 	 */
 	UFUNCTION(BlueprintCallable, Category = Material)
-	static void GetMaterialHierarchyInfo(UMaterialValidationGroup const* InGroup, FString const& InBaseMaterialPath, FMaterialDatabaseAssetHierarchyInfo& OutInfo);
+	static void GetMaterialHierarchyInfo(UMaterialValidationGroup const* InGroup, FSoftObjectPath const& InBaseMaterialPath, FMaterialDatabaseAssetHierarchyInfo& OutInfo);
 
 	/**
 	 * Get all the material properties for a material hierarchy.
diff --git a/Engine/Plugins/Editor/MaterialValidation/Source/Private/MaterialValidationToolkit.cpp b/Engine/Plugins/Editor/MaterialValidation/Source/Private/MaterialValidationToolkit.cpp
--- a/Engine/Plugins/Editor/MaterialValidation/Source/Private/MaterialValidationToolkit.cpp
+++ b/Engine/Plugins/Editor/MaterialValidation/Source/Private/MaterialValidationToolkit.cpp
@@ -1228,7 +1228,7 @@ static void UpdateMaterialInstanceTable(TSharedPtr<SMaterialInstanceTable> InTab
 	const int32 NumMaterialInstances = AssetInfo.MaterialPaths.Num();
 	for (int32 Index = 0; Index < NumMaterialInstances; ++Index)
 	{
-		Rows.Add(MakeShared<FMaterialInstanceRowData>(AssetInfo.MaterialPaths[Index], AssetInfo.MaterialAssetDatas[Index]));
+		Rows.Add(MakeShared<FMaterialInstanceRowData>(AssetInfo.MaterialPaths[Index].ToString(), AssetInfo.MaterialAssetDatas[Index]));
 	}
 
 	// Build all columns.
