MaterialValidationToolkit: Use FSortObjectPath for asset paths.
Add browse to asset buttons in material tables.
Changed material instance sort behavior when sorting by asset. Now we sort by depth first hierarchy order instead of lexical order.
#rnx
#Synced-CL 51325630

diff --git a/Engine/Plugins/Editor/MaterialValidation/Source/MaterialValidation.Build.cs b/Engine/Plugins/Editor/MaterialValidation/Source/MaterialValidation.Build.cs
--- a/Engine/Plugins/Editor/MaterialValidation/Source/MaterialValidation.Build.cs
+++ b/Engine/Plugins/Editor/MaterialValidation/Source/MaterialValidation.Build.cs
@@ -14,6 +14,7 @@ namespace UnrealBuildTool.Rules
 				"DeveloperSettings",
 				"Engine",
 				"Projects",
+				"PropertyEditor",
 				"InputCore",
 				"MaterialEditor",
 				"RHI",
diff --git a/Engine/Plugins/Editor/MaterialValidation/Source/Private/MaterialValidationToolkit.cpp b/Engine/Plugins/Editor/MaterialValidation/Source/Private/MaterialValidationToolkit.cpp
--- a/Engine/Plugins/Editor/MaterialValidation/Source/Private/MaterialValidationToolkit.cpp
+++ b/Engine/Plugins/Editor/MaterialValidation/Source/Private/MaterialValidationToolkit.cpp
@@ -2,12 +2,15 @@
 
 #include "MaterialValidationToolkit.h"
 
+#include "AssetRegistry/IAssetRegistry.h"
 #include "Hash/xxhash.h"
+#include "Framework/Application/SlateApplication.h"
 #include "Materials/MaterialInterface.h"
 #include "MaterialValidationGroup.h"
 #include "MaterialValidationLibrary.h"
 #include "Misc/NotifyHook.h"
 #include "Modules/ModuleManager.h"
+#include "PropertyCustomizationHelpers.h"  
 #include "PropertyEditorModule.h"
 #include "Toolkits/AssetEditorToolkit.h"
 #include "UObject/GCObject.h"
@@ -30,7 +33,7 @@ static float GAssetNameColumnWidth = 700.f;
 /** A data table entry describing a UMaterial. */
 struct FMaterialStats
 {
-	FString AssetPath;
+	FSoftObjectPath AssetPath;
 	int32 MaterialInstanceCount = 0;
 	int32 StaticPermutationCount = 0;
 	int32 ShaderCount = 0;
@@ -45,7 +48,7 @@ struct FMaterialStats
 	static const FName ColumnId_StaticSwitchCount;
 	static const FName ColumnId_ComponentMaskCount;
 
-	FMaterialStats(FString const& InPath, int32 InInstanceCount, int32 InPermutationCount, int32 InShaderCount, int32 InStaticSwitchCount, int32 InComponentMaskCount)
+	FMaterialStats(FSoftObjectPath const& InPath, int32 InInstanceCount, int32 InPermutationCount, int32 InShaderCount, int32 InStaticSwitchCount, int32 InComponentMaskCount)
 		: AssetPath(InPath)
 		, MaterialInstanceCount(InInstanceCount)
 		, StaticPermutationCount(InPermutationCount)
@@ -63,6 +66,47 @@ const FName FMaterialStats::ColumnId_ShaderCount("CID_ShaderCount");
 const FName FMaterialStats::ColumnId_StaticSwitchCount("CID_StaticSwitchCount");
 const FName FMaterialStats::ColumnId_ComponentMaskCount("CID_ComponentMaskCount");
 
+/** Helper function to create an SWidget for AssetPath with text and a browse button. */
+static TSharedRef<SWidget> GenerateAssetPathWidget(FSoftObjectPath InAssetPath)
+{
+	return SNew(SBox)
+		.HAlign(HAlign_Fill)
+		.VAlign(VAlign_Center)
+		[
+			SNew(SHorizontalBox)
+				+ SHorizontalBox::Slot()
+				.FillWidth(1.0f)
+				.VAlign(VAlign_Center)
+				[
+					SNew(STextBlock)
+						.Text(FText::FromString(InAssetPath.ToString()))
+						.ToolTipText(FText::FromString(InAssetPath.ToString()))
+				]
+				+ SHorizontalBox::Slot()
+				.AutoWidth()
+				.VAlign(VAlign_Center)
+				.Padding(4.0f, 0.0f)
+				[
+					PropertyCustomizationHelpers::MakeBrowseButton(
+						FSimpleDelegate::CreateLambda([InAssetPath]()
+						{
+							const FAssetData AssetData = IAssetRegistry::Get()->GetAssetByObjectPath(InAssetPath);
+							
+							const FModifierKeysState KeyState = FSlateApplication::Get().GetModifierKeys();
+							if (KeyState.IsAltDown())
+							{
+								GEditor->EditObject(AssetData.GetAsset());
+							}
+							else
+							{
+								GEditor->SyncBrowserToObjects({ AssetData });
+							}
+						}),
+						LOCTEXT("BrowseTooltip", "Browse to Asset in Content Browser"))
+				]
+		];
+}
+
 /** A slate container of a data table entry describing a UMaterial. */
 class SMaterialTableRow : public SMultiColumnTableRow<TSharedPtr<FMaterialStats>>
 {
@@ -87,7 +131,7 @@ public:
 		{
 			if (ColumnName == FMaterialStats::ColumnId_AssetPath)
 			{
-				return SNew(STextBlock).Text(FText::FromString(RowItem->AssetPath));
+				return GenerateAssetPathWidget(RowItem->AssetPath);
 			}
 			if (ColumnName == FMaterialStats::ColumnId_MaterialInstanceCount)
 			{
@@ -237,7 +281,7 @@ public:
 	{
 		if (SortColumnId != ColumnId)
 		{
-			// Initial sort mode is ascending for text, descending for numeric.
+			// Initial sort mode is ascending for asset path, descending for numeric.
 			SortMode = ColumnId == FMaterialStats::ColumnId_AssetPath ? EColumnSortMode::Ascending : EColumnSortMode::Descending;
 			SortColumnId = ColumnId;
 		}
@@ -257,7 +301,7 @@ public:
 
 		for (const TSharedPtr<FMaterialStats>& Row : AllRows)
 		{
-			if (SearchText.IsEmpty() ||	Row->AssetPath.Contains(SearchText, ESearchCase::IgnoreCase))
+			if (SearchText.IsEmpty() || Row->AssetPath.ToString().Contains(SearchText, ESearchCase::IgnoreCase))
 			{
 				FilteredRows.Add(Row);
 			}
@@ -273,7 +317,7 @@ public:
 		{
 			FilteredRows.Sort([bSortModifier](TSharedPtr<FMaterialStats> const& A, TSharedPtr<FMaterialStats> const& B)
 			{
-				return bSortModifier ? (A->AssetPath > B->AssetPath) : (A->AssetPath < B->AssetPath);
+				return bSortModifier ? A->AssetPath.LexicalLess(B->AssetPath) : B->AssetPath.LexicalLess(A->AssetPath);
 			});
 		}
 		else if (SortColumnId == FMaterialStats::ColumnId_MaterialInstanceCount)
@@ -332,7 +376,7 @@ static TSharedPtr<SMaterialTable> CreateMaterialTable(UMaterialValidationGroup c
 		for (TPair<FString, FMaterialValidationDesc> const& It : InGroup->Materials)
 		{
 			Stats.Add(MakeShared<FMaterialStats>(
-				It.Key, 
+				UMaterialValidationLibrary::ResolveAssetPath(It.Key),
 				It.Value.MaterialInstances.Num(), 
 				It.Value.PermutationHashes.Num() + 1,	// PermutationHashes doesn't include hash for base material 
 				It.Value.NumShadersTotal, 
@@ -387,15 +431,18 @@ static uint32 GetArrayHash(const TArray<uint32>& Array)
 struct FMaterialInstanceRowData
 {
 	/** Path to the UMaterialInstance asset. */
-	FString AssetPath;
+	FSoftObjectPath AssetPath;
+	/** Index in the material instance hierarchy used when sorting by name. */
+	int32 HierarchyIndex = 0;
 	/** Full resolved asset data description for the UMaterialInstance pulled from the material validation database. */
 	FMaterialValidationAssetData_MaterialInstance AssetData;
 
 	/** Per-column-type hash data, derived from instance data */
 	uint32 PerColumnTypeHash[(int32)EMaterialInstanceColumnType::Count] = {};
 
-	FMaterialInstanceRowData(FString const& InPath, FMaterialValidationAssetData_MaterialInstance const& InAssetData)
+	FMaterialInstanceRowData(FSoftObjectPath const& InPath, int32 InHierarchyIndex, FMaterialValidationAssetData_MaterialInstance const& InAssetData)
 		: AssetPath(InPath)
+		, HierarchyIndex(InHierarchyIndex)
 		, AssetData(InAssetData)
 	{
 		PerColumnTypeHash[(int32)EMaterialInstanceColumnType::StaticProperty] = GetNoUsageHash(InAssetData.StaticProperties);
@@ -429,7 +476,8 @@ enum class EMaterialInstanceStaticProperties : int32
 /** Helper getter functions for each of the EMaterialInstanceStaticProperties. */
 struct FMaterialInstanceStaticPropertyGetters
 {
-	static FString const& GetAssetPath(FMaterialInstanceRowData const& P) { return P.AssetPath; }
+	static FSoftObjectPath const& GetAssetPath(FMaterialInstanceRowData const& P) { return P.AssetPath; }
+	static int32 GetAssetPathIndex(FMaterialInstanceRowData const& P) { return P.HierarchyIndex; }
 	static uint32 GetPermutationHash(FMaterialInstanceRowData const& P) { return P.AssetData.PermutationHash; }
 	static uint8 GetTwoSided(FMaterialInstanceRowData const& P) { return P.AssetData.StaticProperties.TwoSided; }
 	static uint8 GetIsThinSurface(FMaterialInstanceRowData const& P) { return P.AssetData.StaticProperties.bIsThinSurface; }
@@ -486,9 +534,9 @@ static TSharedRef<SWidget> GenerateWidgetForStaticProperty(TSharedPtr<FMaterialI
 {
 	using ReturnType = TInvokeResult_T<decltype(GetterPtr), FMaterialInstanceRowData const&>;
 
-	if constexpr(std::is_same<std::remove_cv_t<std::remove_reference_t<ReturnType>>, FString>::value)
+	if constexpr(std::is_same<std::remove_cv_t<std::remove_reference_t<ReturnType>>, FSoftObjectPath>::value)
 	{
-		return SNew(STextBlock).Text(FText::FromString((*GetterPtr)(*RowItem)));
+		return GenerateAssetPathWidget((*GetterPtr)(*RowItem));
 	}
 	else if constexpr(TIsFloatingPoint<ReturnType>::Value)
 	{
@@ -516,7 +564,7 @@ struct FMaterialInstanceStaticPropertyInfo
 
 /** Array for table generation data for EMaterialInstanceStaticProperties. */
 static const FMaterialInstanceStaticPropertyInfo StaticPropertyInfos[(int32)EMaterialInstanceStaticProperties::Count] = {
-	{EMaterialInstanceColumnType::Default, "CID_AssetPath", "Asset Path", &SortByStaticProperty<&FMaterialInstanceStaticPropertyGetters::GetAssetPath>, &GenerateWidgetForStaticProperty<&FMaterialInstanceStaticPropertyGetters::GetAssetPath, &FMaterialInstanceStaticPropertyModifiedGetters::GetAssetPathModified>},
+	{EMaterialInstanceColumnType::Default, "CID_AssetPath", "Asset Path", &SortByStaticProperty<&FMaterialInstanceStaticPropertyGetters::GetAssetPathIndex>, &GenerateWidgetForStaticProperty<&FMaterialInstanceStaticPropertyGetters::GetAssetPath, &FMaterialInstanceStaticPropertyModifiedGetters::GetAssetPathModified>},
 	{EMaterialInstanceColumnType::Default, "CID_PermutationHash", "Permutation Hash", &SortByStaticProperty<&FMaterialInstanceStaticPropertyGetters::GetPermutationHash>, &GenerateWidgetForStaticProperty<&FMaterialInstanceStaticPropertyGetters::GetPermutationHash, &FMaterialInstanceStaticPropertyModifiedGetters::GetPermutationHashModified> },
 	{EMaterialInstanceColumnType::StaticProperty, "CID_TwoSided", "Two Sided", &SortByStaticProperty<&FMaterialInstanceStaticPropertyGetters::GetTwoSided>, &GenerateWidgetForStaticProperty<&FMaterialInstanceStaticPropertyGetters::GetTwoSided, &FMaterialInstanceStaticPropertyModifiedGetters::GetTwoSidedModified> },
 	{EMaterialInstanceColumnType::StaticProperty, "CID_IsThinSurface", "Thin Surface", &SortByStaticProperty<&FMaterialInstanceStaticPropertyGetters::GetIsThinSurface>, &GenerateWidgetForStaticProperty<&FMaterialInstanceStaticPropertyGetters::GetIsThinSurface, &FMaterialInstanceStaticPropertyModifiedGetters::GetIsThinSurfaceModified> },
@@ -1067,7 +1115,7 @@ protected:
 
 		for (const TSharedPtr<FMaterialInstanceRowData>& Row : AllRows)
 		{
-			if (RowSearchText.IsEmpty() || Row->AssetPath.Contains(RowSearchText, ESearchCase::IgnoreCase))
+			if (RowSearchText.IsEmpty() || Row->AssetPath.ToString().Contains(RowSearchText, ESearchCase::IgnoreCase))
 			{
 				if (bUniqueInstances)
 				{
@@ -1130,7 +1178,7 @@ protected:
 	{
 		if (SortColumnId != ColumnId)
 		{
-			// Initial sort mode is ascending for text, descending for numeric.
+			// Initial sort mode is ascending for asset path, descending for numeric.
 			SortMode = ColumnId == GetColumnIdFromStaticProperty(EMaterialInstanceStaticProperties::AssetPath) ? EColumnSortMode::Ascending : EColumnSortMode::Descending;
 			SortColumnId = ColumnId;
 		}
@@ -1217,7 +1265,7 @@ static TSharedPtr<SMaterialInstanceTable> CreateMaterialInstanceTable()
 }
 
 /** Update the table for all material instances in a material. */
-static void UpdateMaterialInstanceTable(TSharedPtr<SMaterialInstanceTable> InTable, UMaterialValidationGroup const* InGroup, FString const& InMaterialPath)
+static void UpdateMaterialInstanceTable(TSharedPtr<SMaterialInstanceTable> InTable, UMaterialValidationGroup const* InGroup, FSoftObjectPath const& InMaterialPath)
 {
 	// Get all the information from the material validation database.
 	FMaterialDatabaseAssetHierarchyInfo AssetInfo;
@@ -1228,7 +1276,7 @@ static void UpdateMaterialInstanceTable(TSharedPtr<SMaterialInstanceTable> InTab
 	const int32 NumMaterialInstances = AssetInfo.MaterialPaths.Num();
 	for (int32 Index = 0; Index < NumMaterialInstances; ++Index)
 	{
-		Rows.Add(MakeShared<FMaterialInstanceRowData>(AssetInfo.MaterialPaths[Index].ToString(), AssetInfo.MaterialAssetDatas[Index]));
+		Rows.Add(MakeShared<FMaterialInstanceRowData>(AssetInfo.MaterialPaths[Index], Index, AssetInfo.MaterialAssetDatas[Index]));
 	}
 
 	// Build all columns.
