MaterialValidation: Fix crash when opening a material in the toolkit where the static parameter counts had changed since the last database build.
#rnx
#Synced-CL 51328950

diff --git a/Engine/Plugins/Editor/MaterialValidation/Source/Private/MaterialValidationLibrary.cpp b/Engine/Plugins/Editor/MaterialValidation/Source/Private/MaterialValidationLibrary.cpp
--- a/Engine/Plugins/Editor/MaterialValidation/Source/Private/MaterialValidationLibrary.cpp
+++ b/Engine/Plugins/Editor/MaterialValidation/Source/Private/MaterialValidationLibrary.cpp
@@ -1293,44 +1293,68 @@ void UMaterialValidationLibrary::GetMaterialHierarchyInfo(UMaterialValidationGro
 		return;
 	}
 
-	FMaterialValidationDesc const* MaterialDesc = MaterialValidationLibrary::FindMaterialDescFromSoftObjectPath(*InGroup, InBaseMaterialPath);
-	if (MaterialDesc == nullptr)
+	const TSoftObjectPtr<UMaterial> BaseMaterialSoftPtr(InBaseMaterialPath);
+	UMaterial const* BaseMaterial = BaseMaterialSoftPtr.LoadSynchronous();
+	if (BaseMaterial == nullptr)
 	{
 		return;
 	}
 
-	TSoftObjectPtr<UMaterial> MaterialSoftPtr(InBaseMaterialPath);
-	UMaterial const* Material = MaterialSoftPtr.LoadSynchronous();
-	if (Material == nullptr)
+	// Get the description stored in the database.
+	FMaterialValidationDesc const* MaterialDesc = MaterialValidationLibrary::FindMaterialDescFromSoftObjectPath(*InGroup, InBaseMaterialPath);
+	if (MaterialDesc == nullptr)
 	{
 		return;
 	}
+	const FMaterialValidationAssetData_Material MaterialAssetData = MaterialDesc->AssetData;
 
-	TArray<FString> SearchPaths;
-	TArray<FString> ExcludePaths;
-	MaterialValidationLibrary::GetSearchPaths(*InGroup, SearchPaths, ExcludePaths);
-
+	// Get the description of the current state.
 	MaterialValidationLibrary::FMaterialPermutationData PermutationData;
-	BuildPermutationData(*Material, PermutationData);
+	BuildPermutationData(*BaseMaterial, PermutationData);
+
+	FMaterialValidationAssetData_Material MaterialAssetDataFromObject;
+	GetAssetData(*BaseMaterial, PermutationData, MaterialAssetDataFromObject);
 
-	OutInfo.StaticSwitchNames.Reserve(PermutationData.StaticSwitchParameterInfo.Num());
-	for (FMaterialParameterInfo const& Info : PermutationData.StaticSwitchParameterInfo)
+	OutInfo.StaticSwitchNames.SetNum(MaterialAssetData.StaticSwitchNum);
+	OutInfo.ComponentMaskNames.SetNum(MaterialAssetData.ComponentMaskNum);
+
+	// Check if the current static parameters are the same as those in the database.
+	// If not then we can't use the found current names (and we don't persist the names in the database).
+	if (MaterialAssetDataFromObject.StaticPropertyLayoutHash == MaterialAssetData.StaticPropertyLayoutHash &&
+		PermutationData.StaticSwitchParameterInfo.Num() == MaterialAssetData.StaticSwitchNum &&
+		PermutationData.ComponentMaskParameterInfo.Num() == MaterialAssetData.ComponentMaskNum)
 	{
-		OutInfo.StaticSwitchNames.Add(Info.Name);
+		for (int32 Index = 0; Index < MaterialAssetData.StaticSwitchNum; ++Index)
+		{
+			OutInfo.StaticSwitchNames[Index] = PermutationData.StaticSwitchParameterInfo[Index].Name;
+		}
+		for (int32 Index = 0; Index < MaterialAssetData.ComponentMaskNum; ++Index)
+		{
+			OutInfo.ComponentMaskNames[Index] = PermutationData.ComponentMaskParameterInfo[Index].Name;
+		}
 	}
-	
-	OutInfo.ComponentMaskNames.Reserve(PermutationData.ComponentMaskParameterInfo.Num());
-	for (FMaterialParameterInfo const& Info : PermutationData.ComponentMaskParameterInfo)
+	else
 	{
-		OutInfo.ComponentMaskNames.Add(Info.Name);
+		for (int32 Index = 0; Index < MaterialAssetData.StaticSwitchNum; ++Index)
+		{
+			static FName BaseName("StaticSwitch");
+			OutInfo.StaticSwitchNames[Index] = FName(BaseName, Index);
+		}
+		for (int32 Index = 0; Index < MaterialAssetData.ComponentMaskNum; ++Index)
+		{
+			static FName BaseName("ComponentMask");
+			OutInfo.ComponentMaskNames[Index] = FName(BaseName, Index);
+		}
 	}
 
-	FMaterialValidationAssetData_Material MaterialAssetData = MaterialDesc->AssetData;
-
 	// Get the material instance hierarchy from the asset registry.
+	TArray<FString> SearchPaths;
+	TArray<FString> ExcludePaths;
+	MaterialValidationLibrary::GetSearchPaths(*InGroup, SearchPaths, ExcludePaths);
+
 	TArray<FName> PackageNames;
 	TArray<int32> HierarchyDepths;
-	MaterialValidationLibrary::GetMaterialInstanceHierarchy(*Material, ExcludePaths, PackageNames, HierarchyDepths);
+	MaterialValidationLibrary::GetMaterialInstanceHierarchy(*BaseMaterial, ExcludePaths, PackageNames, HierarchyDepths);
 
 	// Use a stack of asset data to walk the hierarchy with the asset data overrides applied at each level.
 	TArray<FMaterialValidationAssetData_MaterialInstance> AssetDataStack;
@@ -1339,7 +1363,9 @@ void UMaterialValidationLibrary::GetMaterialHierarchyInfo(UMaterialValidationGro
 	{
 		FMaterialValidationAssetData_MaterialInstance& Root = AssetDataStack.AddDefaulted_GetRef();
 		Root.StaticSwitchOverrideValues = MaterialAssetData.StaticSwitchValues;
+		Root.StaticSwitchOverrideMask.SetNumZeroed(MaterialAssetData.StaticSwitchValues.Num());
 		Root.ComponentMaskOverrideValues = MaterialAssetData.ComponentMaskValues;
+		Root.ComponentMaskOverrideMask.SetNumZeroed(MaterialAssetData.ComponentMaskValues.Num());
 		Root.StaticProperties = MaterialAssetData.StaticProperties;
 		Root.MaterialLayerHash = MaterialAssetData.MaterialLayerHash;
 		Root.PermutationHash = MaterialAssetData.PermutationHash;
