Material Validation: Fix bug when gathering permutation hashes.
Previous static permutation array contents were not being cleared so that incorrect values were returned.
#rnx
#Synced-CL 51328950

diff --git a/Engine/Plugins/Editor/MaterialValidation/Source/Private/MaterialValidationLibrary.cpp b/Engine/Plugins/Editor/MaterialValidation/Source/Private/MaterialValidationLibrary.cpp
--- a/Engine/Plugins/Editor/MaterialValidation/Source/Private/MaterialValidationLibrary.cpp
+++ b/Engine/Plugins/Editor/MaterialValidation/Source/Private/MaterialValidationLibrary.cpp
@@ -160,6 +160,8 @@ static void GetAssetData(
 	FMaterialPermutationData const& InPermutationData, 
 	FMaterialValidationAssetData_Material& OutAssetData)
 {
+	OutAssetData = {};
+
 	// Gather static switches.
 	{
 		const uint32 StaticSwitchNum = InPermutationData.StaticSwitchParameterInfo.Num();
@@ -218,6 +220,8 @@ static void GetAssetData(
 	FMaterialPermutationData const& InPermutationData,
 	FMaterialValidationAssetData_MaterialInstance& OutAssetData)
 {
+	OutAssetData = {};
+
 	UMaterialInterface const* ParentMaterial = InMaterialInstance.Parent;
 
 	// We collect the static switch and component mask values to create the
@@ -708,9 +712,7 @@ static int32 GetNumShaderTypes(UMaterialInterface const& InMaterial)
  */
 static void UpdateValidMaterialInstances(UMaterial const& InMaterial, TConstArrayView<FString> InExcludePaths, FMaterialValidationDesc& OutMaterialDesc)
 {
-	OutMaterialDesc.MaterialInstanceAssetDatas.Reset();
-	OutMaterialDesc.MaterialInstances.Reset();
-	OutMaterialDesc.PermutationHashes.Reset();
+	OutMaterialDesc = {};
 
 	// Initialize shader counts.
 	OutMaterialDesc.NumShadersBase = GetNumShaderTypes(InMaterial);
