MaterialValidation: Return FString instead of FText for material property values.
Also resolve enum values before returning.
#rnx
#Synced-CL 51328950

diff --git a/Engine/Plugins/Editor/MaterialValidation/Source/Private/MaterialValidationLibrary.cpp b/Engine/Plugins/Editor/MaterialValidation/Source/Private/MaterialValidationLibrary.cpp
--- a/Engine/Plugins/Editor/MaterialValidation/Source/Private/MaterialValidationLibrary.cpp
+++ b/Engine/Plugins/Editor/MaterialValidation/Source/Private/MaterialValidationLibrary.cpp
@@ -1470,8 +1470,8 @@ struct FMaterialInstanceStaticPropertyGetters
 	static uint8 GetOutputTranslucentVelocity(FMaterialValidationAssetData_MaterialInstance const& P) { return P.StaticProperties.bOutputTranslucentVelocity; }
 	static uint8 GetHasPixelAnimation(FMaterialValidationAssetData_MaterialInstance const& P) { return P.StaticProperties.bHasPixelAnimation; }
 	static uint8 GetEnableTessellation(FMaterialValidationAssetData_MaterialInstance const& P) { return P.StaticProperties.bEnableTessellation; }
-	static TEnumAsByte<EBlendMode> GetBlendMode(FMaterialValidationAssetData_MaterialInstance const& P) { return P.StaticProperties.BlendMode; }
-	static TEnumAsByte<EMaterialShadingModel> GetShadingModel(FMaterialValidationAssetData_MaterialInstance const& P) { return P.StaticProperties.ShadingModel; }
+	static EBlendMode GetBlendMode(FMaterialValidationAssetData_MaterialInstance const& P) { return P.StaticProperties.BlendMode; }
+	static EMaterialShadingModel GetShadingModel(FMaterialValidationAssetData_MaterialInstance const& P) { return P.StaticProperties.ShadingModel; }
 	static float GetOpacityMaskClipValue(FMaterialValidationAssetData_MaterialInstance const& P) { return P.StaticProperties.OpacityMaskClipValue; }
 };
 
@@ -1491,27 +1491,25 @@ struct FMaterialInstanceStaticPropertyModifiedGetters
 };
 
 template <auto GetterPtr>
-static FText GenerateTextForStaticProperty(FMaterialValidationAssetData_MaterialInstance const& AssetData)
+static FString GetStaticPropertyValueAsString(FMaterialValidationAssetData_MaterialInstance const& AssetData)
 {
 	using ReturnType = TInvokeResult_T<decltype(GetterPtr), FMaterialValidationAssetData_MaterialInstance const&>;
 
-	if constexpr (std::is_same<std::remove_cv_t<std::remove_reference_t<ReturnType>>, FString>::value)
+	if constexpr (TIsEnum<ReturnType>::Value)
 	{
-		return FText::FromString((*GetterPtr)(AssetData));
-	}
-	else if constexpr (TIsFloatingPoint<ReturnType>::Value)
-	{
-		static const FNumberFormattingOptions FormatOptions = FNumberFormattingOptions().SetMaximumFractionalDigits(6);
-		return FText::AsNumber((*GetterPtr)(AssetData), &FormatOptions);
+		return  UEnum::GetDisplayValueAsText((*GetterPtr)(AssetData)).ToString();
 	}
+ 	else if constexpr (TIsFloatingPoint<ReturnType>::Value)
+ 	{
+ 		return FString::Printf(TEXT("%g"), (*GetterPtr)(AssetData));
+ 	}
 	else
 	{
-		static const FNumberFormattingOptions FormatOptions = FNumberFormattingOptions().SetUseGrouping(false);
-		return FText::AsNumber((*GetterPtr)(AssetData), &FormatOptions);
+		return LexToString((*GetterPtr)(AssetData));
 	}
 }
 
-using FnGetPropertyValueText = FText(*)(FMaterialValidationAssetData_MaterialInstance const&);
+using FnGetPropertyValueAsString = FString(*)(FMaterialValidationAssetData_MaterialInstance const&);
 using FnIsPropertyModified = bool(*)(FMaterialValidationAssetData_MaterialInstance const&);
 
 /** Struct containing information needed to generate table entries for EMaterialInstanceStaticProperties. */
@@ -1520,23 +1518,23 @@ struct FMaterialInstanceStaticPropertyInfo
 	FName PropertyId;
 	FName PropertyName;
 	EMaterialInstanceColumnType Category;
-	FnGetPropertyValueText GetPropertyValueTextFunction;
+	FnGetPropertyValueAsString GetPropertyValueAsStringFunction;
 	FnIsPropertyModified GetIsPropertyModifiedFunction;
 };
 
 /** Array for table generation data for EMaterialInstanceStaticProperties. */
 static const FMaterialInstanceStaticPropertyInfo StaticPropertyInfos[(int32)EMaterialInstanceStaticProperties::Count] = {
-	{"PID_PermutationHash", "Permutation Hash", EMaterialInstanceColumnType::Default, &GenerateTextForStaticProperty<&FMaterialInstanceStaticPropertyGetters::GetPermutationHash>, &FMaterialInstanceStaticPropertyModifiedGetters::GetPermutationHashModified },
-	{"PID_TwoSided", "Two Sided", EMaterialInstanceColumnType::StaticProperty, &GenerateTextForStaticProperty <&FMaterialInstanceStaticPropertyGetters::GetTwoSided>, &FMaterialInstanceStaticPropertyModifiedGetters::GetTwoSidedModified },
-	{"PID_IsThinSurface", "Thin Surface", EMaterialInstanceColumnType::StaticProperty, &GenerateTextForStaticProperty <&FMaterialInstanceStaticPropertyGetters::GetIsThinSurface>, &FMaterialInstanceStaticPropertyModifiedGetters::GetIsThinSurfaceModified },
-	{"PID_DitheredLODTransition", "Dithered LOD Transition", EMaterialInstanceColumnType::StaticProperty, &GenerateTextForStaticProperty <&FMaterialInstanceStaticPropertyGetters::GetDitheredLODTransition>, &FMaterialInstanceStaticPropertyModifiedGetters::GetDitheredLODTransitionModified },
-	{"PID_CastDynamicShadowAsMasked", "Dynamic Shadow As Masked", EMaterialInstanceColumnType::StaticProperty, &GenerateTextForStaticProperty <&FMaterialInstanceStaticPropertyGetters::GetCastDynamicShadowAsMasked>, &FMaterialInstanceStaticPropertyModifiedGetters::GetCastDynamicShadowAsMaskedModified },
-	{"PID_OutputTranslucentVelocity", "Output Translucent Velocity", EMaterialInstanceColumnType::StaticProperty, &GenerateTextForStaticProperty <&FMaterialInstanceStaticPropertyGetters::GetOutputTranslucentVelocity>, &FMaterialInstanceStaticPropertyModifiedGetters::GetOutputTranslucentVelocityModified },
-	{"PID_HasPixelAnimation", "Has Pixel Animation", EMaterialInstanceColumnType::StaticProperty, &GenerateTextForStaticProperty <&FMaterialInstanceStaticPropertyGetters::GetHasPixelAnimation>, &FMaterialInstanceStaticPropertyModifiedGetters::GetHasPixelAnimationModified },
-	{"PID_EnableTessellation", "Tessellation", EMaterialInstanceColumnType::StaticProperty, &GenerateTextForStaticProperty <&FMaterialInstanceStaticPropertyGetters::GetEnableTessellation>, &FMaterialInstanceStaticPropertyModifiedGetters::GetEnableTessellationModified },
-	{"PID_BlendMode", "Blend Mode", EMaterialInstanceColumnType::StaticProperty, &GenerateTextForStaticProperty <&FMaterialInstanceStaticPropertyGetters::GetBlendMode>, &FMaterialInstanceStaticPropertyModifiedGetters::GetBlendModeModified },
-	{"PID_ShadingModel", "Shading Model", EMaterialInstanceColumnType::StaticProperty, &GenerateTextForStaticProperty <&FMaterialInstanceStaticPropertyGetters::GetShadingModel>, &FMaterialInstanceStaticPropertyModifiedGetters::GetShadingModelModified },
-	{"PID_OpacityMaskClipValue", "Opacity Mask Clip Value", EMaterialInstanceColumnType::StaticProperty, &GenerateTextForStaticProperty <&FMaterialInstanceStaticPropertyGetters::GetOpacityMaskClipValue>, &FMaterialInstanceStaticPropertyModifiedGetters::GetOpacityMaskClipModified },
+	{"PID_PermutationHash", "Permutation Hash", EMaterialInstanceColumnType::Default, &GetStaticPropertyValueAsString<&FMaterialInstanceStaticPropertyGetters::GetPermutationHash>, &FMaterialInstanceStaticPropertyModifiedGetters::GetPermutationHashModified },
+	{"PID_TwoSided", "Two Sided", EMaterialInstanceColumnType::StaticProperty, &GetStaticPropertyValueAsString <&FMaterialInstanceStaticPropertyGetters::GetTwoSided>, &FMaterialInstanceStaticPropertyModifiedGetters::GetTwoSidedModified },
+	{"PID_IsThinSurface", "Thin Surface", EMaterialInstanceColumnType::StaticProperty, &GetStaticPropertyValueAsString <&FMaterialInstanceStaticPropertyGetters::GetIsThinSurface>, &FMaterialInstanceStaticPropertyModifiedGetters::GetIsThinSurfaceModified },
+	{"PID_DitheredLODTransition", "Dithered LOD Transition", EMaterialInstanceColumnType::StaticProperty, &GetStaticPropertyValueAsString <&FMaterialInstanceStaticPropertyGetters::GetDitheredLODTransition>, &FMaterialInstanceStaticPropertyModifiedGetters::GetDitheredLODTransitionModified },
+	{"PID_CastDynamicShadowAsMasked", "Dynamic Shadow As Masked", EMaterialInstanceColumnType::StaticProperty, &GetStaticPropertyValueAsString <&FMaterialInstanceStaticPropertyGetters::GetCastDynamicShadowAsMasked>, &FMaterialInstanceStaticPropertyModifiedGetters::GetCastDynamicShadowAsMaskedModified },
+	{"PID_OutputTranslucentVelocity", "Output Translucent Velocity", EMaterialInstanceColumnType::StaticProperty, &GetStaticPropertyValueAsString <&FMaterialInstanceStaticPropertyGetters::GetOutputTranslucentVelocity>, &FMaterialInstanceStaticPropertyModifiedGetters::GetOutputTranslucentVelocityModified },
+	{"PID_HasPixelAnimation", "Has Pixel Animation", EMaterialInstanceColumnType::StaticProperty, &GetStaticPropertyValueAsString <&FMaterialInstanceStaticPropertyGetters::GetHasPixelAnimation>, &FMaterialInstanceStaticPropertyModifiedGetters::GetHasPixelAnimationModified },
+	{"PID_EnableTessellation", "Tessellation", EMaterialInstanceColumnType::StaticProperty, &GetStaticPropertyValueAsString <&FMaterialInstanceStaticPropertyGetters::GetEnableTessellation>, &FMaterialInstanceStaticPropertyModifiedGetters::GetEnableTessellationModified },
+	{"PID_BlendMode", "Blend Mode", EMaterialInstanceColumnType::StaticProperty, &GetStaticPropertyValueAsString <&FMaterialInstanceStaticPropertyGetters::GetBlendMode>, &FMaterialInstanceStaticPropertyModifiedGetters::GetBlendModeModified },
+	{"PID_ShadingModel", "Shading Model", EMaterialInstanceColumnType::StaticProperty, &GetStaticPropertyValueAsString <&FMaterialInstanceStaticPropertyGetters::GetShadingModel>, &FMaterialInstanceStaticPropertyModifiedGetters::GetShadingModelModified },
+	{"PID_OpacityMaskClipValue", "Opacity Mask Clip Value", EMaterialInstanceColumnType::StaticProperty, &GetStaticPropertyValueAsString <&FMaterialInstanceStaticPropertyGetters::GetOpacityMaskClipValue>, &FMaterialInstanceStaticPropertyModifiedGetters::GetOpacityMaskClipModified },
 };
 
 /** Static array of column FNames for all material usages. */
@@ -1614,7 +1612,7 @@ void UMaterialValidationLibrary::GetMaterialPropertyValues(FMaterialDatabaseAsse
 	for (MaterialValidationLibrary::FMaterialInstanceStaticPropertyInfo const& StaticPropertyInfo : MaterialValidationLibrary::StaticPropertyInfos)
 	{
 		OutValues.Add({
-			.Value = StaticPropertyInfo.GetPropertyValueTextFunction(AssetData), 
+			.Value = StaticPropertyInfo.GetPropertyValueAsStringFunction(AssetData), 
 			.bIsModifiedFromParent = StaticPropertyInfo.GetIsPropertyModifiedFunction(AssetData) });
 	}
 	
@@ -1622,7 +1620,7 @@ void UMaterialValidationLibrary::GetMaterialPropertyValues(FMaterialDatabaseAsse
 	{
 		const bool bValue = (AssetData.StaticProperties.UsageFlags & (1u << Index)) != 0;
 		const bool bModified = (AssetData.StaticPropertyOverrideFlags.Override_UsageFlags & (1u << Index)) != 0;
-		OutValues.Add({	.Value = FText::AsNumber(bValue ? 1 : 0), .bIsModifiedFromParent = bModified });
+		OutValues.Add({	.Value = bValue ? TEXT("1") : TEXT("0"), .bIsModifiedFromParent = bModified});
 	}
 
 	for (int32 Index = 0; Index < InInfo.StaticSwitchNames.Num(); ++Index)
@@ -1634,12 +1632,12 @@ void UMaterialValidationLibrary::GetMaterialPropertyValues(FMaterialDatabaseAsse
 		{
 			const bool bValue = (AssetData.StaticSwitchOverrideValues[BitArrayIndex] & (1u << BitIndex)) != 0;
 			const bool bModified = (AssetData.StaticSwitchOverrideMask[BitArrayIndex] & (1u << BitIndex)) != 0;
-			OutValues.Add({ .Value = FText::AsNumber(bValue ? 1 : 0), .bIsModifiedFromParent = bModified });
+			OutValues.Add({ .Value = bValue ? TEXT("1") : TEXT("0"), .bIsModifiedFromParent = bModified });
 		}
 		else
 		{
 			// Mismatch possibly caused by add/removal of parameters since database build.
-			OutValues.Add({ .Value = FText::FromString("?"), .bIsModifiedFromParent = false });
+			OutValues.Add({ .Value = TEXT("?"), .bIsModifiedFromParent = false });
 		}
 	}
 
@@ -1657,12 +1655,12 @@ void UMaterialValidationLibrary::GetMaterialPropertyValues(FMaterialDatabaseAsse
 			RGBA.AppendChar(Value & 2 ? 'G' : '_');
 			RGBA.AppendChar(Value & 4 ? 'B' : '_');
 			RGBA.AppendChar(Value & 8 ? 'A' : '_');
-			OutValues.Add({ .Value = FText::FromString(*RGBA), .bIsModifiedFromParent = bModified });
+			OutValues.Add({ .Value = *RGBA, .bIsModifiedFromParent = bModified });
 		}
 		else
 		{
 			// Mismatch possibly caused by add/removal of parameters since database build.
-			OutValues.Add({ .Value = FText::FromString("?"), .bIsModifiedFromParent = false });
+			OutValues.Add({ .Value = TEXT("?"), .bIsModifiedFromParent = false });
 		}
 	}
 }
diff --git a/Engine/Plugins/Editor/MaterialValidation/Source/Private/MaterialValidationLibrary.h b/Engine/Plugins/Editor/MaterialValidation/Source/Private/MaterialValidationLibrary.h
--- a/Engine/Plugins/Editor/MaterialValidation/Source/Private/MaterialValidationLibrary.h
+++ b/Engine/Plugins/Editor/MaterialValidation/Source/Private/MaterialValidationLibrary.h
@@ -69,9 +69,9 @@ struct FMaterialDatabaseAssetPropertyValue
 {
 	GENERATED_BODY()
 
-	/** Text value for the property. */
+	/** String value for the property. */
 	UPROPERTY(BlueprintReadOnly, Category = Material)
-	FText Value;
+	FString Value;
 	/** true if the value for the property is different from the value on its parent material. */
 	UPROPERTY(BlueprintReadOnly, Category = Material)
 	bool bIsModifiedFromParent = false;
