Material Validation: Experiments harness and simple experiments

#tests local editor tests of MVG asset interface
#Synced-CL 51353528

diff --git a/Engine/Plugins/Editor/MaterialValidation/Source/Private/MaterialValidationAssetData.h b/Engine/Plugins/Editor/MaterialValidation/Source/Private/MaterialValidationAssetData.h
--- a/Engine/Plugins/Editor/MaterialValidation/Source/Private/MaterialValidationAssetData.h
+++ b/Engine/Plugins/Editor/MaterialValidation/Source/Private/MaterialValidationAssetData.h
@@ -136,3 +136,13 @@ struct FMaterialValidationAssetData_MaterialInstance
 	UPROPERTY()
 	uint32 PermutationHash = 0;
 };
+
+namespace MaterialValidation
+{
+	uint32 BuildPermutationHash(
+		FStaticPermutationProperties const& InStaticProperties,
+		TConstArrayView<uint32> InStaticSwitchValues,
+		TConstArrayView<uint32> InComponentMaskValues,
+		uint32 InMaterialLayerHash
+	);
+}
diff --git a/Engine/Plugins/Editor/MaterialValidation/Source/Private/MaterialValidationGroup.cpp b/Engine/Plugins/Editor/MaterialValidation/Source/Private/MaterialValidationGroup.cpp
--- a/Engine/Plugins/Editor/MaterialValidation/Source/Private/MaterialValidationGroup.cpp
+++ b/Engine/Plugins/Editor/MaterialValidation/Source/Private/MaterialValidationGroup.cpp
@@ -146,3 +146,20 @@ void UMaterialValidationGroup::UpdatePermutations()
 {
 	UMaterialValidationLibrary::UpdateMaterialPermutationsInGroup(this);
 }
+
+/** Return a hash built for a material or material instance which will match for all instances that generate the same shader permutations. */
+uint32 MaterialValidation::BuildPermutationHash(FStaticPermutationProperties const& InStaticProperties, TConstArrayView<uint32> InStaticSwitchValues, TConstArrayView<uint32> InComponentMaskValues, uint32 InMaterialLayerHash)
+{
+	uint32 Hash = GetTypeHash(InStaticProperties);
+	for (uint32 Value : InStaticSwitchValues)
+	{
+		Hash = HashCombineFast(Hash, GetTypeHash(Value));
+	}
+	for (uint32 Value : InComponentMaskValues)
+	{
+		Hash = HashCombineFast(Hash, GetTypeHash(Value));
+	}
+	Hash = HashCombineFast(Hash, InMaterialLayerHash);
+	return Hash;
+}
+
diff --git a/Engine/Plugins/Editor/MaterialValidation/Source/Private/MaterialValidationLibrary.cpp b/Engine/Plugins/Editor/MaterialValidation/Source/Private/MaterialValidationLibrary.cpp
--- a/Engine/Plugins/Editor/MaterialValidation/Source/Private/MaterialValidationLibrary.cpp
+++ b/Engine/Plugins/Editor/MaterialValidation/Source/Private/MaterialValidationLibrary.cpp
@@ -138,22 +138,6 @@ static uint32 BuildLayerHash(UMaterialInterface const& InMaterial)
 	return 0;
 }
 
-/** Return a hash built for a material or material instance which will match for all instances that generate the same shader permutations. */
-static uint32 BuildPermutationHash(FStaticPermutationProperties const& InStaticProperties, TConstArrayView<uint32> InStaticSwitchValues, TConstArrayView<uint32> InComponentMaskValues, uint32 InMaterialLayerHash)
-{
-	uint32 Hash = GetTypeHash(InStaticProperties);
-	for (uint32 Value : InStaticSwitchValues)
-	{
-		Hash = HashCombineFast(Hash, GetTypeHash(Value));
-	} 
-	for (uint32 Value : InComponentMaskValues)
-	{
-		Hash = HashCombineFast(Hash, GetTypeHash(Value));
-	}
-	Hash = HashCombineFast(Hash, InMaterialLayerHash);
-	return Hash;
-}
-
 /** Extract AssetData from a UMaterial object. */
 static void GetAssetData(
 	UMaterial const& InMaterial, 
@@ -211,7 +195,7 @@ static void GetAssetData(
 	OutAssetData.MaterialLayerHash = BuildLayerHash(InMaterial);
 	
 	OutAssetData.StaticPropertyLayoutHash = InPermutationData.StaticPropertyLayoutHash;
-	OutAssetData.PermutationHash = BuildPermutationHash(OutAssetData.StaticProperties, OutAssetData.StaticSwitchValues, OutAssetData.ComponentMaskValues, OutAssetData.MaterialLayerHash);
+	OutAssetData.PermutationHash = MaterialValidation::BuildPermutationHash(OutAssetData.StaticProperties, OutAssetData.StaticSwitchValues, OutAssetData.ComponentMaskValues, OutAssetData.MaterialLayerHash);
 }
 
 /** Extract AssetData from a UMaterialInstance object. */
@@ -302,7 +286,7 @@ static void GetAssetData(
 
 	OutAssetData.MaterialLayerHash = BuildLayerHash(InMaterialInstance);
 	
-	OutAssetData.PermutationHash = BuildPermutationHash(StaticProperties, StaticSwitchValues, ComponentMaskValues, OutAssetData.MaterialLayerHash);
+	OutAssetData.PermutationHash = MaterialValidation::BuildPermutationHash(StaticProperties, StaticSwitchValues, ComponentMaskValues, OutAssetData.MaterialLayerHash);
 }
 
 /** Build hash for a material instance which will match for all instances that generate the same shader permutations. */
@@ -1262,7 +1246,7 @@ int32 UMaterialValidationLibrary::GetModifiedShaderCount(UMaterialValidationGrou
 		Child.MaterialLayerHash = AssetData.MaterialLayerHash;
 
 		// Build the permutation hash and if this is a new permutation then add its shader count.
-		Child.PermutationHash = MaterialValidationLibrary::BuildPermutationHash(Child.StaticProperties, Child.StaticSwitchOverrideValues, Child.ComponentMaskOverrideValues, Child.MaterialLayerHash);
+		Child.PermutationHash = MaterialValidation::BuildPermutationHash(Child.StaticProperties, Child.StaticSwitchOverrideValues, Child.ComponentMaskOverrideValues, Child.MaterialLayerHash);
 		if (!PermutationHashes.Contains(Child.PermutationHash))
 		{
 			// If we don't have the material instance loaded then create a transient version.
@@ -1430,7 +1414,7 @@ void UMaterialValidationLibrary::GetMaterialHierarchyInfo(UMaterialValidationGro
 		MaterialValidationLibrary::ApplyBitMaskOverrides(Parent.ComponentMaskOverrideValues, AssetData.ComponentMaskOverrideMask, AssetData.ComponentMaskOverrideValues, Child.ComponentMaskOverrideValues, Child.ComponentMaskOverrideMask);
 		MaterialValidationLibrary::ApplyBasePropertyOverrides(Parent.StaticProperties, AssetData.StaticPropertyOverrideFlags, AssetData.StaticProperties, Child.StaticProperties, Child.StaticPropertyOverrideFlags);
 		Child.MaterialLayerHash = AssetData.MaterialLayerHash;
-		Child.PermutationHash = MaterialValidationLibrary::BuildPermutationHash(Child.StaticProperties, Child.StaticSwitchOverrideValues, Child.ComponentMaskOverrideValues, Child.MaterialLayerHash);
+		Child.PermutationHash = MaterialValidation::BuildPermutationHash(Child.StaticProperties, Child.StaticSwitchOverrideValues, Child.ComponentMaskOverrideValues, Child.MaterialLayerHash);
 
 		OutInfo.MaterialPaths.Add(MaterialValidationLibrary::ConvertStringPathToSoftObjectPath(PackagePath));
 		OutInfo.MaterialAssetDatas.Add(Child);
diff --git a/Engine/Plugins/Editor/MaterialValidation/Source/Private/MaterialValidationToolkit.cpp b/Engine/Plugins/Editor/MaterialValidation/Source/Private/MaterialValidationToolkit.cpp
--- a/Engine/Plugins/Editor/MaterialValidation/Source/Private/MaterialValidationToolkit.cpp
+++ b/Engine/Plugins/Editor/MaterialValidation/Source/Private/MaterialValidationToolkit.cpp
@@ -66,6 +66,120 @@ const FName FMaterialStats::ColumnId_ShaderCount("CID_ShaderCount");
 const FName FMaterialStats::ColumnId_StaticSwitchCount("CID_StaticSwitchCount");
 const FName FMaterialStats::ColumnId_ComponentMaskCount("CID_ComponentMaskCount");
 
+struct FExperimentData
+{
+	FString Name;
+	uint32 ReductionCount = 0;
+	float ReductionPercent = 0.0f;
+
+	static const FName ColumnId_Name;
+	static const FName ColumnId_ReductionCount;
+	static const FName ColumnId_ReductionPercent;
+};
+
+const FName FExperimentData::ColumnId_Name("CID_ExperimentName");
+const FName FExperimentData::ColumnId_ReductionCount("CID_ReductionCount");
+const FName FExperimentData::ColumnId_ReductionPercent("CID_ReductionPercent");
+
+class SExperimentTableRow : public SMultiColumnTableRow<TSharedPtr<FExperimentData>>
+{
+	TSharedPtr<FExperimentData> RowItem;
+
+public:
+	SLATE_BEGIN_ARGS(SExperimentTableRow) {}
+		SLATE_ARGUMENT(TSharedPtr<FExperimentData>, Item)
+	SLATE_END_ARGS()
+
+	void Construct(const FArguments& InArgs, const TSharedRef<STableViewBase>& InOwnerTable)
+	{
+		RowItem = InArgs._Item;
+		SMultiColumnTableRow<TSharedPtr<FExperimentData>>::Construct(FSuperRowType::FArguments(), InOwnerTable);
+	}
+
+	virtual TSharedRef<SWidget> GenerateWidgetForColumn(const FName& ColumnName) override
+	{
+		if (RowItem.IsValid())
+		{
+			if (ColumnName == FExperimentData::ColumnId_Name)
+			{
+				return SNew(STextBlock).Text(FText::FromString(RowItem->Name));
+			}
+			if (ColumnName == FExperimentData::ColumnId_ReductionCount)
+			{
+				return SNew(STextBlock).Text(FText::AsNumber(RowItem->ReductionCount));
+			}
+			if (ColumnName == FExperimentData::ColumnId_ReductionPercent)
+			{
+				return SNew(STextBlock).Text(FText::AsPercent(RowItem->ReductionPercent));
+			}
+		}
+		return SNew(STextBlock).Text(FText::FromString("?"));
+	}
+};
+
+class SExperimentsTable : public SCompoundWidget
+{
+	TSharedPtr<SListView<TSharedPtr<FExperimentData>>> ListView;
+	TSharedPtr<SHeaderRow> HeaderRow;
+	TArray<TSharedPtr<FExperimentData>> Rows;
+	float NameColumnWidth = 500.0f;
+
+public:
+	SLATE_BEGIN_ARGS(SExperimentsTable) {}
+	SLATE_END_ARGS()
+
+	void Construct(const FArguments& InArgs)
+	{
+		HeaderRow = SNew(SHeaderRow)
+			+ SHeaderRow::Column(FExperimentData::ColumnId_Name)
+				.HAlignHeader(HAlign_Center).HAlignCell(HAlign_Left)
+				.DefaultLabel(LOCTEXT("ExperimentName", "Experiment"))
+				.ManualWidth(NameColumnWidth)
+				.ManualWidth_Lambda([this]() { return NameColumnWidth; })
+				.OnWidthChanged_Lambda([this](float NewWidth) { NameColumnWidth = NewWidth; })
+			+ SHeaderRow::Column(FExperimentData::ColumnId_ReductionCount)
+				.HAlignHeader(HAlign_Center).HAlignCell(HAlign_Center)
+				.DefaultLabel(LOCTEXT("ExperimentReductionCount", "Reduction"))
+				.FixedWidth(120)
+			+ SHeaderRow::Column(FExperimentData::ColumnId_ReductionPercent)
+				.HAlignHeader(HAlign_Center).HAlignCell(HAlign_Center)
+				.DefaultLabel(LOCTEXT("ExperimentReductionPercent", "Reduction %"))
+				.FixedWidth(120);
+
+		ChildSlot
+		[
+			SNew(SVerticalBox)
+			+SVerticalBox::Slot().HAlign(HAlign_Fill).FillHeight(1.0f)
+			[
+				SNew(SScrollBox)
+				.Orientation(Orient_Horizontal)
+				+SScrollBox::Slot()
+				[
+					SAssignNew(ListView, SListView<TSharedPtr<FExperimentData>>)
+					.HeaderRow(HeaderRow.ToSharedRef())
+					.ListItemsSource(&Rows)
+					.OnGenerateRow_Lambda([this](TSharedPtr<FExperimentData> InItem, const TSharedRef<STableViewBase>& InOwnerTable)
+						{
+							return SNew(SExperimentTableRow, InOwnerTable).Item(InItem);
+						})
+				]
+			]
+		];
+	}
+
+	/** Set the table content. */
+	void SetContent(TArray<TSharedPtr<FExperimentData>> const& InRows)
+	{
+		Rows = InRows;
+		ListView->RequestListRefresh();
+	}
+};
+
+static TSharedPtr<SExperimentsTable> CreateExperimentsTable()
+{
+	return SNew(SExperimentsTable);
+}
+
 /** Helper function to create an SWidget for AssetPath with text and a browse button. */
 static TSharedRef<SWidget> GenerateAssetPathWidget(FSoftObjectPath InAssetPath)
 {
@@ -1264,8 +1378,60 @@ static TSharedPtr<SMaterialInstanceTable> CreateMaterialInstanceTable()
 	return SNew(SMaterialInstanceTable).InitialColumns(Columns).InitialRows(Rows).InitialParameterNames(ParameterNames);
 }
 
+struct FExperimentKey
+{
+	FStaticPermutationProperties StaticProperties;
+	TArray<uint32> SwitchValues;
+	TArray<uint32> MaskValues;
+	uint32 MaterialLayerHash;
+};
+
+static FExperimentKey MakeExperimentKey(FMaterialValidationAssetData_MaterialInstance const& AssetData)
+{
+	return { AssetData.StaticProperties, AssetData.StaticSwitchOverrideValues, AssetData.ComponentMaskOverrideValues, AssetData.MaterialLayerHash };
+}
+
+static uint32 HashExperimentKey(FExperimentKey const& Key)
+{
+	return MaterialValidation::BuildPermutationHash(Key.StaticProperties, Key.SwitchValues, Key.MaskValues, Key.MaterialLayerHash);
+}
+
+using FExperimentKeyModifier = TFunction<FExperimentKey(FExperimentKey const&)>;
+
+static void RunExperiment(
+	FString const& Name,
+	TArray<TSharedPtr<FMaterialInstanceRowData>> const& InstanceRows,
+	int32 BaselineCount,
+	TArray<TSharedPtr<FExperimentData>>& OutputRows,
+	FExperimentKeyModifier Modifier)
+{
+	if (BaselineCount > 0 && InstanceRows.Num())
+	{
+		TSet<uint32> HypotheticalHashes;
+		for (auto const& Row : InstanceRows)
+		{
+			HypotheticalHashes.Add(HashExperimentKey(Modifier(MakeExperimentKey(Row->AssetData))));
+		}
+
+		const uint32 Reduction = FMath::Max(0, BaselineCount - (int32)HypotheticalHashes.Num());
+
+		if (Reduction > 0)
+		{
+			auto Result = MakeShared<FExperimentData>();
+			Result->Name = Name;
+			Result->ReductionCount = Reduction;
+			Result->ReductionPercent = (float)Reduction / BaselineCount;
+			OutputRows.Add(Result);
+		}
+	}
+}
+
 /** Update the table for all material instances in a material. */
-static void UpdateMaterialInstanceTable(TSharedPtr<SMaterialInstanceTable> InTable, UMaterialValidationGroup const* InGroup, FSoftObjectPath const& InMaterialPath)
+static void UpdateMaterialInstanceTables(
+	TSharedPtr<SMaterialInstanceTable> InMaterialTable,
+	TSharedPtr<SExperimentsTable> InExperimentsTable,
+	UMaterialValidationGroup const* InGroup, 
+	FSoftObjectPath const& InMaterialPath)
 {
 	// Get all the information from the material validation database.
 	FMaterialDatabaseAssetHierarchyInfo AssetInfo;
@@ -1330,7 +1496,89 @@ static void UpdateMaterialInstanceTable(TSharedPtr<SMaterialInstanceTable> InTab
 	}
 
 	// Apply to the table.
-	InTable->SetContent(Rows, Columns, ParameterNames);
+	InMaterialTable->SetContent(Rows, Columns, ParameterNames);
+
+	// do Experiments
+	TArray<TSharedPtr<FExperimentData>> ExperimentRows;
+
+	if (Rows.Num())
+	{
+		TSet<uint32> BaselineHashes;
+		for (auto const& Row : Rows)
+		{ 
+			BaselineHashes.Add(HashExperimentKey(MakeExperimentKey(Row->AssetData)));
+		}
+		const int32 BaselineCount = BaselineHashes.Num();
+
+		RunExperiment(TEXT("No-op"), Rows, BaselineCount, ExperimentRows,
+			[](FExperimentKey Key) { return Key; });
+
+		// No-op should add no rows
+		ensure(ExperimentRows.Num() == 0);
+
+		RunExperiment(TEXT("Normalize OpacityMaskClipValue"), Rows, BaselineCount, ExperimentRows,
+			[](FExperimentKey Key) { Key.StaticProperties.OpacityMaskClipValue = 0.f; return Key; });
+
+		for (int32 SwitchIndex = 0; SwitchIndex < ParameterNames->StaticSwitches.Num(); ++SwitchIndex)
+		{
+			RunExperiment(
+				FString::Printf(TEXT("Eliminate Switch %s"), *ParameterNames->StaticSwitches[SwitchIndex].ToString()),
+				Rows, BaselineCount, ExperimentRows,
+				[SwitchIndex](FExperimentKey Key)
+				{
+					const uint32 DwordIndex = SwitchIndex / 32;
+					const uint32 BitIndex = SwitchIndex % 32;
+					if (DwordIndex < (uint32)Key.SwitchValues.Num())
+					{
+						Key.SwitchValues[DwordIndex] &= ~(1u << BitIndex);
+					}
+					return Key;
+				});
+		}
+
+		if (!ParameterNames->ComponentMasks.IsEmpty())
+		{
+			// Elimination of single-channel component masks, these all can be ChannelMaskParameter which has no permutation cost
+			TArray<uint32> const& BaseMaskValues = Rows[0]->AssetData.ComponentMaskOverrideValues;
+			const int32 NumMasks = ParameterNames->ComponentMasks.Num();
+
+			TArray<uint32> SingleChannelClearMask;
+			SingleChannelClearMask.SetNumZeroed(BaseMaskValues.Num());
+
+			for (int32 Index = 0; Index < NumMasks; ++Index)
+			{
+				const uint32 DwordIndex = Index / 8;
+				const uint32 BitIndex = (Index * 4) % 32;
+				if (DwordIndex < (uint32)BaseMaskValues.Num())
+				{
+					const uint32 Nibble = (BaseMaskValues[DwordIndex] >> BitIndex) & 0xf;
+					const bool bSingleChannel = Nibble != 0 && (Nibble & (Nibble - 1)) == 0;
+					if (bSingleChannel)
+					{
+						SingleChannelClearMask[DwordIndex] |= (0xfu << BitIndex);
+					}
+				}
+			}
+
+			RunExperiment(TEXT("Single-channel Component Masks to ChannelMaskParameter"), Rows, BaselineCount, ExperimentRows,
+				[SingleChannelClearMask](FExperimentKey Key)
+				{
+					for (int32 i = 0; i < Key.MaskValues.Num() && i < SingleChannelClearMask.Num(); ++i)
+					{
+						Key.MaskValues[i] &= ~SingleChannelClearMask[i];
+					}
+					return Key;
+				});
+		}
+	}
+
+	// sort by best percentage win
+	ExperimentRows.Sort([](TSharedPtr<FExperimentData> const& A, TSharedPtr<FExperimentData> const& B)
+		{
+			return A->ReductionPercent > B->ReductionPercent;
+		});
+
+	InExperimentsTable->SetContent(ExperimentRows);
 }
 
 } // namespace MaterialValidation
@@ -1344,11 +1592,11 @@ class FMaterialValidationToolkit : public FAssetEditorToolkit, public FGCObject,
 	/** The current UMaterialValidationGroup being edited. */
 	TObjectPtr<UMaterialValidationGroup> EditGroup;
 
-	/** Tab manager for the toolkit. */
-	TSharedPtr<class FTabManager> TabManager;
 	/** Menu category for the toolkit tabs. */
 	TSharedPtr<class FWorkspaceItem> WorkspaceMenuCategory;
 
+	TSharedPtr<MaterialValidation::SExperimentsTable> ExperimentsTable;
+
 	/** Standard details panel view. */
 	TSharedPtr<class IDetailsView> DetailsView;
 	/** Material table widget. */
@@ -1360,6 +1608,7 @@ class FMaterialValidationToolkit : public FAssetEditorToolkit, public FGCObject,
 	static const FName TabId_Details;
 	static const FName TabId_Materials;
 	static const FName TabId_MaterialInstances;
+	static const FName TabId_Experiments;
 
 public:
 	// Begin FGCObject overrides.
@@ -1413,6 +1662,11 @@ public:
 			.SetDisplayName(LOCTEXT("Toolkit_MaterialInstancesTab_DisplayName", "Material Instances"))
 			.SetGroup(WorkspaceMenuCategory.ToSharedRef())
 			.SetIcon(FSlateIcon(FAppStyle::GetAppStyleSetName(), "LevelEditor.Tabs.StatsViewer"));
+
+		InTabManager->RegisterTabSpawner(TabId_Experiments, FOnSpawnTab::CreateSP(this, &FMaterialValidationToolkit::SpawnExperimentsTab))
+			.SetDisplayName(LOCTEXT("Toolkit_ExperimentsTab_DisplayName", "Experiments"))
+			.SetGroup(WorkspaceMenuCategory.ToSharedRef())
+			.SetIcon(FSlateIcon(FAppStyle::GetAppStyleSetName(), "LevelEditor.Tabs.StatsViewer"));
 	}
 
 	virtual void UnregisterTabSpawners(const TSharedRef<class FTabManager>& InTabManager) override
@@ -1421,6 +1675,7 @@ public:
 		InTabManager->UnregisterTabSpawner(TabId_Details);
 		InTabManager->UnregisterTabSpawner(TabId_Materials);
 		InTabManager->UnregisterTabSpawner(TabId_MaterialInstances);
+		InTabManager->UnregisterTabSpawner(TabId_Experiments);
 	}
 	// End FAssetEditorToolkit overrides.
 
@@ -1442,6 +1697,15 @@ public:
 			];
 	}
 
+	TSharedRef<SDockTab> SpawnExperimentsTab(const FSpawnTabArgs& Args)
+	{
+		return SNew(SDockTab)
+		.Label(LOCTEXT("Toolkit_Experiments_Label", "Experiments"))
+		[
+			ExperimentsTable.ToSharedRef()
+		];
+	}
+
 	TSharedRef<SDockTab> SpawnMaterialInstancesTab(const FSpawnTabArgs& Args)
 	{
 		return SNew(SDockTab)
@@ -1467,8 +1731,9 @@ public:
 		MaterialTable = MaterialValidation::CreateMaterialTable(EditGroup, MaterialValidation::FOnRowItemSelect::CreateSP(this, &FMaterialValidationToolkit::OnRowItemSelect));
 
 		MaterialInstanceTable = MaterialValidation::CreateMaterialInstanceTable();
+		ExperimentsTable = MaterialValidation::CreateExperimentsTable();
 
-		const TSharedRef<FTabManager::FLayout> Layout = FTabManager::NewLayout("MaterialValidationGroupLayout_v4")
+		const TSharedRef<FTabManager::FLayout> Layout = FTabManager::NewLayout("MaterialValidationGroupLayout_v5")
 			->AddArea
 			(
 				FTabManager::NewPrimaryArea()->SetOrientation(Orient_Horizontal)
@@ -1481,7 +1746,8 @@ public:
 					->Split(FTabManager::NewStack()
 						->AddTab(TabId_Materials, ETabState::OpenedTab))
 					->Split(FTabManager::NewStack()
-						->AddTab(TabId_MaterialInstances, ETabState::ClosedTab))
+						->AddTab(TabId_MaterialInstances, ETabState::ClosedTab)
+						->AddTab(TabId_Experiments, ETabState::ClosedTab))
 				)
 			);
 
@@ -1491,16 +1757,26 @@ public:
 	/** Apply selected material to the material instances table. */
 	void OnRowItemSelect(TSharedPtr<MaterialValidation::FMaterialStats> InMaterialStats)
 	{
-		MaterialValidation::UpdateMaterialInstanceTable(MaterialInstanceTable, EditGroup, InMaterialStats->AssetPath);
+		MaterialValidation::UpdateMaterialInstanceTables(MaterialInstanceTable, ExperimentsTable, EditGroup, InMaterialStats->AssetPath);
 
-		InvokeTab(TabId_MaterialInstances);
+		TSharedPtr<SDockTab> InstancesTab = GetTabManager()->FindExistingLiveTab(TabId_MaterialInstances);
+		TSharedPtr<SDockTab> ExperimentsTab = GetTabManager()->FindExistingLiveTab(TabId_Experiments);
+
+		const bool bExperimentsForeground = ExperimentsTab.IsValid() && ExperimentsTab->IsForeground();
+
+		// Open any closed tabs
+		if (!InstancesTab.IsValid()) { InvokeTab(TabId_MaterialInstances); }
+		if (!ExperimentsTab.IsValid()) { InvokeTab(TabId_Experiments); }
+
+		// Restore foreground — Experiments if it was active, Instances otherwise
+		InvokeTab(bExperimentsForeground ? TabId_Experiments : TabId_MaterialInstances);
 	}
 };
 
 const FName FMaterialValidationToolkit::TabId_Details("MainTab");
 const FName FMaterialValidationToolkit::TabId_Materials("MaterialsTab");
 const FName FMaterialValidationToolkit::TabId_MaterialInstances("MaterialInstancesTab");
-
+const FName FMaterialValidationToolkit::TabId_Experiments("ExperimentsTab");
 
 FText UAssetDefinition_MaterialValidationGroup::GetAssetDisplayName() const 
 {
