MaterialValidation: Replace slow path finding project plugins.
Was doing expensive FPath logic, but there is  IPlugin::GetLoadedFrom() which does the same much faster.
#rnx
#Synced-CL 51537754

diff --git a/Engine/Plugins/Editor/MaterialValidation/Source/Private/MaterialValidationLibrary.cpp b/Engine/Plugins/Editor/MaterialValidation/Source/Private/MaterialValidationLibrary.cpp
--- a/Engine/Plugins/Editor/MaterialValidation/Source/Private/MaterialValidationLibrary.cpp
+++ b/Engine/Plugins/Editor/MaterialValidation/Source/Private/MaterialValidationLibrary.cpp
@@ -378,12 +378,11 @@ static void GetSearchPaths(UMaterialValidationGroup const& InGroup, TArray<FStri
 		TArray<TSharedRef<IPlugin>> Plugins = PluginManager.GetDiscoveredPlugins();
 		for (TSharedRef<IPlugin> const& Plugin : Plugins)
 		{
-			if (Plugin->CanContainContent())
+			if (Plugin->GetLoadedFrom() == EPluginLoadedFrom::Project && Plugin->CanContainContent())
 			{
-				if (FPaths::IsUnderDirectory(Plugin->GetContentDir(), FPaths::ProjectDir()))
-				{
-					OutPaths.Add(FString::Printf(TEXT("/%s/"), *Plugin->GetName()));
-				}
+				TStringBuilder<128> Builder;
+				Builder.AppendChar(TEXT('/')).Append(Plugin->GetName()).AppendChar(TEXT('/'));
+				OutPaths.Add(*Builder);
 			}
 		}
 	}
