MaterialValidation: Add custom diff tool for UMaterialValidationGroup assets in source control.
The diff window shows two levels: a filterable summary table of all changed base materials and a tree view of material instance property diffs for the selected material.
#Synced-CL 51579155

diff --git a/Engine/Plugins/Editor/MaterialValidation/Source/Private/MaterialValidationGroupDiff.cpp b/Engine/Plugins/Editor/MaterialValidation/Source/Private/MaterialValidationGroupDiff.cpp
new file mode 100644
--- /dev/null
+++ b/Engine/Plugins/Editor/MaterialValidation/Source/Private/MaterialValidationGroupDiff.cpp
@@ -0,0 +1,896 @@
+// Copyright Epic Games, Inc. All Rights Reserved.
+
+#include "MaterialValidationGroupDiff.h"
+
+#include "AssetRegistry/IAssetRegistry.h"
+#include "Editor.h"
+#include "Framework/Application/SlateApplication.h"
+#include "MaterialValidationGroup.h"
+#include "MaterialValidationLibrary.h"
+#include "MaterialValidationLibraryTypes.h"
+#include "MaterialValidationToolkit.h"
+#include "PropertyCustomizationHelpers.h"
+#include "Widgets/Input/SButton.h"
+#include "Widgets/Input/SSearchBox.h"
+#include "Widgets/Layout/SSplitter.h"
+#include "Widgets/Text/STextBlock.h"
+#include "Widgets/Views/SHeaderRow.h"
+#include "Widgets/Views/SListView.h"
+#include "Widgets/Views/STableRow.h"
+#include "Widgets/Views/STreeView.h"
+
+#define LOCTEXT_NAMESPACE "MaterialValidation"
+
+namespace MaterialValidation {
+
+/** Value to track user selected width of asset name columns. */
+static float GAssetNameColumnWidth_GroupDiff = 700.f;
+
+/** Get color to use according to whether cost delta is positive, negative or neutral. */
+static FSlateColor GetDeltaColor(int32 InDelta)
+{
+	if (InDelta > 0)
+	{
+		return FSlateColor(FLinearColor(1.f, 0.4f, 0.4f));
+	}
+	if (InDelta < 0)
+	{
+		return FSlateColor(FLinearColor(0.4f, 1.f, 0.4f));
+	}
+	return FSlateColor::UseForeground();
+}
+
+/** Return FText that formats the number with a positive/negative sign. */
+static FText GetTextSignedInt(int32 Value)
+{
+	return (Value > 0) ? FText::Format(LOCTEXT("PositiveNumberFormat", "+{0}"),	FText::AsNumber(Value)) : FText::AsNumber(Value);
+}
+
+static const FText GEmptyText = FText::FromString("-");
+static const FText GUnknownText = FText::FromString("?");
+
+/** Per-material row data for the material diff table. Holds the counts and deltas from both snapshots. */
+struct FMaterialDiffStats
+{
+	FSoftObjectPath AssetPath;
+	EMaterialInstanceDiffType ChangeType = EMaterialInstanceDiffType::Modified;
+	
+	int32 OldInstanceCount = 0;
+	int32 NewInstanceCount = 0;
+	int32 InstanceDelta = 0;
+	int32 OldShaderCount = 0;
+	int32 NewShaderCount = 0;
+	int32 ShaderDelta = 0;
+	int32 OldPermutationCount = 0;
+	int32 NewPermutationCount = 0;
+	int32 PermutationDelta = 0;
+
+	static const FName ColumnId_AssetPath;
+	static const FName ColumnId_OldInstanceCount;
+	static const FName ColumnId_NewInstanceCount;
+	static const FName ColumnId_InstanceDelta;
+	static const FName ColumnId_OldShaderCount;
+	static const FName ColumnId_NewShaderCount;
+	static const FName ColumnId_ShaderDelta;
+	static const FName ColumnId_OldPermutationCount;
+	static const FName ColumnId_NewPermutationCount;
+	static const FName ColumnId_PermutationDelta;
+};
+
+const FName FMaterialDiffStats::ColumnId_AssetPath("CID_Diff_AssetPath");
+const FName FMaterialDiffStats::ColumnId_OldInstanceCount("CID_Diff_OldInstanceCount");
+const FName FMaterialDiffStats::ColumnId_NewInstanceCount("CID_Diff_NewInstanceCount");
+const FName FMaterialDiffStats::ColumnId_InstanceDelta("CID_Diff_InstanceDelta");
+const FName FMaterialDiffStats::ColumnId_OldShaderCount("CID_Diff_OldShaderCount");
+const FName FMaterialDiffStats::ColumnId_NewShaderCount("CID_Diff_NewShaderCount");
+const FName FMaterialDiffStats::ColumnId_ShaderDelta("CID_Diff_ShaderDelta");
+const FName FMaterialDiffStats::ColumnId_OldPermutationCount("CID_Diff_OldPermCount");
+const FName FMaterialDiffStats::ColumnId_NewPermutationCount("CID_Diff_NewPermCount");
+const FName FMaterialDiffStats::ColumnId_PermutationDelta("CID_Diff_PermDelta");
+
+/** Build diff rows by comparing two materials groups. */
+static TArray<TSharedPtr<FMaterialDiffStats>> BuildDiffRows(UMaterialValidationGroup const& InOldGroup, UMaterialValidationGroup const& InNewGroup)
+{
+	TArray<TSharedPtr<FMaterialDiffStats>> Rows;
+
+	TSet<FString> AllPaths;
+	for (auto const& Pair : InOldGroup.Materials) 
+	{
+		AllPaths.Add(Pair.Key); 
+	}
+	for (auto const& Pair : InNewGroup.Materials) 
+	{
+		AllPaths.Add(Pair.Key); 
+	}
+
+	for (FString const& PathStr : AllPaths)
+	{
+		FMaterialValidationDesc const* OldDesc = InOldGroup.Materials.Find(PathStr);
+		FMaterialValidationDesc const* NewDesc = InNewGroup.Materials.Find(PathStr);
+
+		TSharedPtr<FMaterialDiffStats> Row = MakeShared<FMaterialDiffStats>();
+		Row->AssetPath = UMaterialValidationLibrary::ResolveAssetPath(PathStr);
+
+		if (OldDesc == nullptr)
+		{
+			Row->ChangeType = EMaterialInstanceDiffType::Added;
+			Row->NewInstanceCount = NewDesc->MaterialInstances.Num();
+			Row->NewShaderCount = NewDesc->NumShadersTotal;
+			Row->NewPermutationCount = NewDesc->PermutationHashes.Num() + 1;
+			Row->ShaderDelta = Row->NewShaderCount;
+			Row->InstanceDelta = Row->NewInstanceCount;
+			Row->PermutationDelta = Row->NewPermutationCount;
+		}
+		else if (NewDesc == nullptr)
+		{
+			Row->ChangeType = EMaterialInstanceDiffType::Removed;
+			Row->OldInstanceCount = OldDesc->MaterialInstances.Num();
+			Row->OldShaderCount = OldDesc->NumShadersTotal;
+			Row->OldPermutationCount = OldDesc->PermutationHashes.Num() + 1;
+			Row->ShaderDelta = -Row->OldShaderCount;
+			Row->InstanceDelta = -Row->OldInstanceCount;
+			Row->PermutationDelta = -Row->OldPermutationCount;
+		}
+		else
+		{
+			Row->ChangeType = EMaterialInstanceDiffType::Modified;
+			Row->OldInstanceCount = OldDesc->MaterialInstances.Num();
+			Row->NewInstanceCount = NewDesc->MaterialInstances.Num();
+			Row->OldShaderCount = OldDesc->NumShadersTotal;
+			Row->NewShaderCount = NewDesc->NumShadersTotal;
+			Row->OldPermutationCount = OldDesc->PermutationHashes.Num() + 1;
+			Row->NewPermutationCount = NewDesc->PermutationHashes.Num() + 1;
+			Row->InstanceDelta = Row->NewInstanceCount - Row->OldInstanceCount;
+			Row->ShaderDelta = Row->NewShaderCount - Row->OldShaderCount;
+			Row->PermutationDelta = Row->NewPermutationCount - Row->OldPermutationCount;
+
+			// Skip unchanged materials.
+			if (Row->InstanceDelta == 0 && Row->ShaderDelta == 0 && Row->PermutationDelta == 0)
+			{
+				continue;
+			}
+		}
+
+		Rows.Add(MoveTemp(Row));
+	}
+
+	return Rows;
+}
+
+/** Table row widget for one material in the material diff table. */
+class SMaterialDiffTableRow : public SMultiColumnTableRow<TSharedPtr<FMaterialDiffStats>>
+{
+	TSharedPtr<FMaterialDiffStats> RowItem;
+
+public:
+	SLATE_BEGIN_ARGS(SMaterialDiffTableRow) {}
+		SLATE_ARGUMENT(TSharedPtr<FMaterialDiffStats>, Item)
+	SLATE_END_ARGS()
+
+	void Construct(const FArguments& InArgs, const TSharedRef<STableViewBase>& InOwnerTable)
+	{
+		RowItem = InArgs._Item;
+		SMultiColumnTableRow<TSharedPtr<FMaterialDiffStats>>::Construct(FSuperRowType::FArguments(), InOwnerTable);
+	}
+
+	virtual TSharedRef<SWidget> GenerateWidgetForColumn(const FName& ColumnName) override
+	{
+		if (!RowItem.IsValid())
+		{
+			return SNew(STextBlock).Text(GUnknownText);
+		}
+
+		const bool bAdded = RowItem->ChangeType == EMaterialInstanceDiffType::Added;
+		const bool bHasOld = !bAdded;
+		const bool bRemoved = RowItem->ChangeType == EMaterialInstanceDiffType::Removed;
+		const bool bHasNew = !bRemoved;
+
+		// Helper to generate a numeric cell.
+		auto MakeNumericCell = [](int32 InValue, bool bShow) -> TSharedRef<SWidget>
+		{
+			return SNew(SBox).HAlign(HAlign_Center).VAlign(VAlign_Center)
+			[
+				SNew(STextBlock).Text(bShow ? FText::AsNumber(InValue) : GEmptyText)
+			];
+		};
+
+		// Helper to generate a delta cell.
+		auto MakeDeltaCell = [](int32 InDelta, bool bBothSidesExist) -> TSharedRef<SWidget>
+		{
+			const FText DeltaText = bBothSidesExist ? GetTextSignedInt(InDelta) : GEmptyText;
+			const FSlateColor Color = bBothSidesExist ? GetDeltaColor(InDelta) : GetDeltaColor(0);
+			return SNew(SBox).HAlign(HAlign_Center).VAlign(VAlign_Center)
+			[
+				SNew(STextBlock).Text(DeltaText).ColorAndOpacity(Color)
+			];
+		};
+
+		if (ColumnName == FMaterialDiffStats::ColumnId_AssetPath)
+		{
+			const FSlateColor PathColor = bAdded ? GetDeltaColor(1) : bRemoved ? GetDeltaColor(-1) : GetDeltaColor(0);
+			return MaterialValidation::GenerateAssetPathWidget(RowItem->AssetPath, PathColor);
+		}
+		if (ColumnName == FMaterialDiffStats::ColumnId_OldInstanceCount) { return MakeNumericCell(RowItem->OldInstanceCount, bHasOld); }
+		if (ColumnName == FMaterialDiffStats::ColumnId_NewInstanceCount) { return MakeNumericCell(RowItem->NewInstanceCount, bHasNew); }
+		if (ColumnName == FMaterialDiffStats::ColumnId_InstanceDelta) { return MakeDeltaCell(RowItem->InstanceDelta, bHasOld && bHasNew); }
+		if (ColumnName == FMaterialDiffStats::ColumnId_OldShaderCount) { return MakeNumericCell(RowItem->OldShaderCount, bHasOld); }
+		if (ColumnName == FMaterialDiffStats::ColumnId_NewShaderCount) { return MakeNumericCell(RowItem->NewShaderCount, bHasNew); }
+		if (ColumnName == FMaterialDiffStats::ColumnId_ShaderDelta) { return MakeDeltaCell(RowItem->ShaderDelta, bHasOld && bHasNew); }
+		if (ColumnName == FMaterialDiffStats::ColumnId_OldPermutationCount) { return MakeNumericCell(RowItem->OldPermutationCount, bHasOld); }
+		if (ColumnName == FMaterialDiffStats::ColumnId_NewPermutationCount) { return MakeNumericCell(RowItem->NewPermutationCount, bHasNew); }
+		if (ColumnName == FMaterialDiffStats::ColumnId_PermutationDelta) { return MakeDeltaCell(RowItem->PermutationDelta, bHasOld && bHasNew); }
+
+		return SNew(STextBlock).Text(GUnknownText);
+	}
+};
+
+/** Delegate for notification of material data row selection. */
+DECLARE_DELEGATE_OneParam(FOnMaterialDiffRowSelected, TSharedPtr<FMaterialDiffStats>);
+
+/**
+ * Widget for the material diff table. 
+ * Shows one sortable/filterable row for each base material that changed between the two group snapshots.
+ * Fires FOnMaterialDiffRowSelected when the user clicks a row to populate the material instance tree.
+ */
+class SMaterialDiffTable : public SCompoundWidget
+{
+	/** All the available rows. */
+	TArray<TSharedPtr<FMaterialDiffStats>> AllRows;
+	/** An ordered and filtered subset of AllRows. */
+	TArray<TSharedPtr<FMaterialDiffStats>> FilteredRows;
+	/** A slate view of the table. */
+	TSharedPtr<SListView<TSharedPtr<FMaterialDiffStats>>> ListView;
+
+	/** Search text used for filtering the rows. */
+	FString SearchText;
+	/** Selected column name for sorting the table. */
+	FName SortColumnId = NAME_None;
+	/** Selected mode for sorting the table. */
+	EColumnSortMode::Type SortMode = EColumnSortMode::None;
+
+	/** Delegate for row selection. */
+	FOnMaterialDiffRowSelected OnRowSelected;
+
+public:
+	SLATE_BEGIN_ARGS(SMaterialDiffTable) {}
+		SLATE_ARGUMENT(TArray<TSharedPtr<FMaterialDiffStats>>, InitialRows)
+		SLATE_EVENT(FOnMaterialDiffRowSelected, OnRowSelected)
+	SLATE_END_ARGS()
+
+	/** Builds the slate table view. */
+	void Construct(const FArguments& InArgs)
+	{
+		AllRows = InArgs._InitialRows;
+		FilteredRows = InArgs._InitialRows;
+		OnRowSelected = InArgs._OnRowSelected;
+
+		SortColumnId = FMaterialDiffStats::ColumnId_ShaderDelta;
+		SortMode = EColumnSortMode::Descending;
+		UpdateSort();
+
+		ChildSlot
+		[
+			SNew(SVerticalBox)
+			+ SVerticalBox::Slot().AutoHeight().Padding(2)
+			[
+				SNew(SSearchBox)
+					.HintText(LOCTEXT("DiffSearchHint", "Filter By Asset"))
+					.OnTextChanged(this, &SMaterialDiffTable::OnSearchTextChanged)
+			]
+			+ SVerticalBox::Slot().FillHeight(1.f).Padding(2)
+			[
+				SAssignNew(ListView, SListView<TSharedPtr<FMaterialDiffStats>>)
+				.HeaderRow(BuildHeaderRow())
+				.ListItemsSource(&FilteredRows)
+				.OnGenerateRow_Lambda([](TSharedPtr<FMaterialDiffStats> InItem, const TSharedRef<STableViewBase>& InOwnerTable)
+				{
+					return SNew(SMaterialDiffTableRow, InOwnerTable).Item(InItem);
+				})
+				.SelectionMode(ESelectionMode::Single)
+				.OnSelectionChanged_Lambda([this](TSharedPtr<FMaterialDiffStats> InItem, ESelectInfo::Type)
+				{
+					OnRowSelected.ExecuteIfBound(InItem);
+				})
+			]
+		];
+	}
+
+private:
+	/** Builds the header row for the slate table view. */
+	TSharedRef<SHeaderRow> BuildHeaderRow()
+	{
+		auto MakeSortableCol = [this](FName InId, FText InLabel) -> SHeaderRow::FColumn::FArguments
+		{
+			return SHeaderRow::Column(InId)
+				.HAlignHeader(HAlign_Center)
+				.HAlignCell(HAlign_Center)
+				.DefaultLabel(InLabel)
+				.OnSort(this, &SMaterialDiffTable::OnColumnSortModeChanged)
+				.FixedWidth(140);
+		};
+
+		return SNew(SHeaderRow)
+			+ SHeaderRow::Column(FMaterialDiffStats::ColumnId_AssetPath)
+				.HAlignHeader(HAlign_Center).HAlignCell(HAlign_Left)
+				.DefaultLabel(LOCTEXT("DiffAssetPath", "Asset Path"))
+				.OnSort(this, &SMaterialDiffTable::OnColumnSortModeChanged)
+				.ManualWidth(GAssetNameColumnWidth_GroupDiff)
+				.ManualWidth_Lambda([this]() 
+				{
+					return GAssetNameColumnWidth_GroupDiff;
+				})
+				.OnWidthChanged_Lambda([this](float NewWidth) 
+				{
+					GAssetNameColumnWidth_GroupDiff = NewWidth;
+				})
+			+ MakeSortableCol(FMaterialDiffStats::ColumnId_ShaderDelta, LOCTEXT("ShaderDelta", "Shaders (Delta)"))
+			+ MakeSortableCol(FMaterialDiffStats::ColumnId_InstanceDelta, LOCTEXT("InstanceDelta", "Instances (Delta)"))
+			+ MakeSortableCol(FMaterialDiffStats::ColumnId_PermutationDelta, LOCTEXT("PermDelta", "Permutations (Delta)"))
+			+ MakeSortableCol(FMaterialDiffStats::ColumnId_OldShaderCount, LOCTEXT("OldShaders", "Shaders (Old)"))
+			+ MakeSortableCol(FMaterialDiffStats::ColumnId_NewShaderCount, LOCTEXT("NewShaders", "Shaders (New)"))
+			+ MakeSortableCol(FMaterialDiffStats::ColumnId_OldInstanceCount, LOCTEXT("OldInstances", "Instances (Old)"))
+			+ MakeSortableCol(FMaterialDiffStats::ColumnId_NewInstanceCount, LOCTEXT("NewInstances", "Instances (New)"))
+			+ MakeSortableCol(FMaterialDiffStats::ColumnId_OldPermutationCount, LOCTEXT("OldPermutations", "Permutations (Old)"))
+			+ MakeSortableCol(FMaterialDiffStats::ColumnId_NewPermutationCount, LOCTEXT("NewPermutations", "Permutations (New)"));
+	}
+
+	/** Called when modifying the search text. */
+	void OnSearchTextChanged(const FText& InText)
+	{
+		SearchText = InText.ToString();
+		UpdateFilter();
+		UpdateSort();
+		ListView->RequestListRefresh();
+	}
+
+	/** Called when clicking on a column header to sort. */
+	void OnColumnSortModeChanged(const EColumnSortPriority::Type, const FName& ColumnId, const EColumnSortMode::Type InSortMode)
+	{
+		if (SortColumnId != ColumnId)
+		{
+			SortMode = (ColumnId == FMaterialDiffStats::ColumnId_AssetPath) ? EColumnSortMode::Ascending : EColumnSortMode::Descending;
+			SortColumnId = ColumnId;
+		}
+		else
+		{
+			SortMode = (SortMode == EColumnSortMode::Descending) ? EColumnSortMode::Ascending : EColumnSortMode::Descending;
+		}
+		UpdateSort();
+		ListView->RequestListRefresh();
+	}
+
+	/**  Apply filter to AllRows to recreate FilteredRows. */
+	void UpdateFilter()
+	{
+		FilteredRows.Reset();
+		for (TSharedPtr<FMaterialDiffStats> const& Row : AllRows)
+		{
+			if (SearchText.IsEmpty() || Row->AssetPath.ToString().Contains(SearchText, ESearchCase::IgnoreCase))
+			{
+				FilteredRows.Add(Row);
+			}
+		}
+	}
+
+	/**  Apply sort to FilteredRows. */
+	void UpdateSort()
+	{
+		const bool bSortModifier = (SortMode == EColumnSortMode::Descending);
+
+		if (SortColumnId == FMaterialDiffStats::ColumnId_AssetPath)
+		{
+			FilteredRows.Sort([bSortModifier](auto const& A, auto const& B)
+			{ 
+				return bSortModifier ? A->AssetPath.LexicalLess(B->AssetPath) : B->AssetPath.LexicalLess(A->AssetPath); 
+			});
+		}
+		else
+		{
+			// For delta columns sort by signed value; for raw count columns sort by value.
+			auto GetSortKey = [](TSharedPtr<FMaterialDiffStats> const& R, FName ColumnId) -> int32
+			{
+				if (ColumnId == FMaterialDiffStats::ColumnId_ShaderDelta) { return R->ShaderDelta; }
+				if (ColumnId == FMaterialDiffStats::ColumnId_InstanceDelta) { return R->InstanceDelta; }
+				if (ColumnId == FMaterialDiffStats::ColumnId_PermutationDelta) { return R->PermutationDelta; }
+				if (ColumnId == FMaterialDiffStats::ColumnId_OldShaderCount) { return R->OldShaderCount; }
+				if (ColumnId == FMaterialDiffStats::ColumnId_NewShaderCount) { return R->NewShaderCount; }
+				if (ColumnId == FMaterialDiffStats::ColumnId_OldInstanceCount) { return R->OldInstanceCount; }
+				if (ColumnId == FMaterialDiffStats::ColumnId_NewInstanceCount) { return R->NewInstanceCount; }
+				if (ColumnId == FMaterialDiffStats::ColumnId_OldPermutationCount) { return R->OldPermutationCount; }
+				if (ColumnId == FMaterialDiffStats::ColumnId_NewPermutationCount) { return R->NewPermutationCount; }
+				return 0;
+			};
+
+			// Returns true when the cell for this row/column displays "-" (the value is not applicable).
+			// Added rows have no Old data; Removed rows have no New data; both show "-" in delta columns.
+			// Rows with blank cells should always sort last.
+			auto IsBlankCell = [](TSharedPtr<FMaterialDiffStats> const& R, FName ColumnId) -> bool
+			{
+				const bool bAdded   = R->ChangeType == EMaterialInstanceDiffType::Added;
+				const bool bRemoved = R->ChangeType == EMaterialInstanceDiffType::Removed;
+
+				if (!bAdded && !bRemoved)
+				{
+					return false;
+				}
+				if (ColumnId == FMaterialDiffStats::ColumnId_ShaderDelta
+					|| ColumnId == FMaterialDiffStats::ColumnId_InstanceDelta
+					|| ColumnId == FMaterialDiffStats::ColumnId_PermutationDelta)
+				{
+					return true;
+				}
+				if (bAdded)
+				{
+					return ColumnId == FMaterialDiffStats::ColumnId_OldShaderCount
+						|| ColumnId == FMaterialDiffStats::ColumnId_OldInstanceCount
+						|| ColumnId == FMaterialDiffStats::ColumnId_OldPermutationCount;
+				}
+				return ColumnId == FMaterialDiffStats::ColumnId_NewShaderCount
+					|| ColumnId == FMaterialDiffStats::ColumnId_NewInstanceCount
+					|| ColumnId == FMaterialDiffStats::ColumnId_NewPermutationCount;
+			};
+
+			const FName ColId = SortColumnId;
+			FilteredRows.Sort([bSortModifier, ColId, &GetSortKey, &IsBlankCell](auto const& A, auto const& B)
+			{
+				const bool bABlank = IsBlankCell(A, ColId);
+				const bool bBBlank = IsBlankCell(B, ColId);
+				if (bABlank != bBBlank)
+				{
+					// Blank cells always sort after real values, regardless of sort direction.
+					return bBBlank;
+				}
+				if (bABlank && bBBlank)
+				{
+					// Both cells are blank — use asset path as a stable secondary sort.
+					return A->AssetPath.LexicalLess(B->AssetPath);
+				}
+				return bSortModifier ? GetSortKey(A, ColId) > GetSortKey(B, ColId) : GetSortKey(A, ColId) < GetSortKey(B, ColId);
+			});
+		}
+	}
+};
+
+/** Node in the material instance property diff tree. */
+struct FMaterialDiffTreeNode
+{
+	/** true if this is a property node. */
+	bool bIsPropertyNode = false;
+	/** Populated for material instance(parent) nodes. */
+	FMaterialInstanceDiffResult DiffResult;
+	/** Populated for property nodes. */
+	FMaterialInstancePropertyDiff PropertyDiff;
+	/** Children properties of this node. */
+	TArray<TSharedPtr<FMaterialDiffTreeNode>> Children;
+};
+
+/** Table row for the material instance property diff tree.  */
+class SMaterialInstanceDiffTreeRow : public SMultiColumnTableRow<TSharedPtr<FMaterialDiffTreeNode>>
+{
+	/** The node for the row. */
+	TSharedPtr<FMaterialDiffTreeNode> Node;
+
+public:
+	static const FName ColAssetPath;
+	static const FName ColProperty;
+	static const FName ColOldValue;
+	static const FName ColNewValue;
+
+	SLATE_BEGIN_ARGS(SMaterialInstanceDiffTreeRow) {}
+		SLATE_ARGUMENT(TSharedPtr<FMaterialDiffTreeNode>, Node)
+	SLATE_END_ARGS()
+
+	/** Builds the slate node view. */
+	void Construct(const FArguments& InArgs, const TSharedRef<STableViewBase>& InOwnerTableView)
+	{
+		Node = InArgs._Node;
+		SMultiColumnTableRow<TSharedPtr<FMaterialDiffTreeNode>>::Construct(FSuperRowType::FArguments(), InOwnerTableView);
+	}
+
+	virtual TSharedRef<SWidget> GenerateWidgetForColumn(const FName& InColumnName) override
+	{
+		if (!Node.IsValid())
+		{
+			return SNullWidget::NullWidget;
+		}
+
+		// Property difference row.
+		// Asset path column is empty and property/value columns are filled.
+		if (Node->bIsPropertyNode)
+		{
+			if (InColumnName == ColProperty)
+			{
+				return SNew(SHorizontalBox)
+					+ SHorizontalBox::Slot().FillWidth(1.f).HAlign(HAlign_Center).VAlign(VAlign_Center)
+					[
+						SNew(STextBlock).Text(FText::FromName(Node->PropertyDiff.PropertyName))
+					];
+			}
+			if (InColumnName == ColOldValue)
+			{
+				return SNew(SBox).HAlign(HAlign_Center).VAlign(VAlign_Center)
+				[
+					SNew(STextBlock).Text(FText::FromString(Node->PropertyDiff.OldValue))
+				];
+			}
+			if (InColumnName == ColNewValue)
+			{
+				return SNew(SBox).HAlign(HAlign_Center).VAlign(VAlign_Center)
+				[
+					SNew(STextBlock).Text(FText::FromString(Node->PropertyDiff.NewValue))
+				];
+			}
+			return SNullWidget::NullWidget;
+		}
+
+		// Material instance (parent) row.
+		// Show expander arrow and coloured asset path in the asset path column only.
+		if (InColumnName == ColAssetPath)
+		{
+			FSlateColor NodeColor = GetDeltaColor(0);
+			if (Node->DiffResult.DiffType == EMaterialInstanceDiffType::Added)
+			{
+				NodeColor = GetDeltaColor(1);
+			}
+			else if (Node->DiffResult.DiffType == EMaterialInstanceDiffType::Removed)
+			{
+				NodeColor = GetDeltaColor(-1);
+			}
+			return SNew(SHorizontalBox)
+				+ SHorizontalBox::Slot().AutoWidth().VAlign(VAlign_Center)
+				[
+					SNew(SExpanderArrow, SharedThis(this)).IndentAmount(12)
+				]
+				+ SHorizontalBox::Slot().FillWidth(1.f)
+				[
+					MaterialValidation::GenerateAssetPathWidget(Node->DiffResult.InstancePath, NodeColor)
+				];
+		}
+
+		return SNullWidget::NullWidget;
+	}
+};
+
+const FName SMaterialInstanceDiffTreeRow::ColAssetPath = "AssetPath";
+const FName SMaterialInstanceDiffTreeRow::ColProperty  = "Property";
+const FName SMaterialInstanceDiffTreeRow::ColOldValue  = "OldValue";
+const FName SMaterialInstanceDiffTreeRow::ColNewValue  = "NewValue";
+
+/**
+ * Material instance property diff tree. 
+ * This is populated by GetMaterialValidationDescDiff when the user selects a row in SMaterialDiffTable.
+ * It shows Added/Removed instances as leaf nodes and Modified instances as expandable nodes with property diff children.
+ */
+class SMaterialInstanceDiffTree : public SCompoundWidget
+{
+	/** The older group snapshot. Held as a weak pointer since this widget is not a UObject. */
+	TWeakObjectPtr<const UMaterialValidationGroup> OldGroup;
+	/** The newer group snapshot. Held as a weak pointer since this widget is not a UObject. */
+	TWeakObjectPtr<const UMaterialValidationGroup> NewGroup;
+	/** Path of the base material currently shown in the tree. */
+	FSoftObjectPath SelectedMaterialPath;
+
+	/** Full unfiltered list of parent (instance) nodes built from the last diff call. */
+	TArray<TSharedPtr<FMaterialDiffTreeNode>> AllRootNodes;
+	/** Subset of AllRootNodes that pass the current search filters. */
+	TArray<TSharedPtr<FMaterialDiffTreeNode>> FilteredRootNodes;
+	/** The tree view widget. */
+	TSharedPtr<STreeView<TSharedPtr<FMaterialDiffTreeNode>>> TreeView;
+
+	/** Current text from the instance path search box. */
+	FString InstanceSearchText;
+	/** Current text from the property name search box. */
+	FString PropertySearchText;
+
+	/** Tracks whether all tree nodes are currently expanded, toggled by clicking the Asset Path column header. */
+	bool bAllExpanded = false;
+
+public:
+	SLATE_BEGIN_ARGS(SMaterialInstanceDiffTree) {}
+		SLATE_ARGUMENT(UMaterialValidationGroup const*, OldGroup)
+		SLATE_ARGUMENT(UMaterialValidationGroup const*, NewGroup)
+	SLATE_END_ARGS()
+
+	/** Builds the slate view. */
+	void Construct(const FArguments& InArgs)
+	{
+		OldGroup = InArgs._OldGroup;
+		NewGroup = InArgs._NewGroup;
+
+		ChildSlot
+		[
+			SNew(SVerticalBox)
+			+ SVerticalBox::Slot().AutoHeight().Padding(2)
+			[
+				SNew(SHorizontalBox)
+				+ SHorizontalBox::Slot().FillWidth(1.f).Padding(0, 0, 4, 0)
+				[
+					SNew(SSearchBox)
+						.HintText(LOCTEXT("InstanceSearchHint", "Filter By Instance Name"))
+						.OnTextChanged(this, &SMaterialInstanceDiffTree::OnInstanceSearchChanged)
+				]
+				+ SHorizontalBox::Slot().FillWidth(1.f)
+				[
+					SNew(SSearchBox)
+						.HintText(LOCTEXT("PropertySearchHint", "Filter By Property Name"))
+						.OnTextChanged(this, &SMaterialInstanceDiffTree::OnPropertySearchChanged)
+				]
+			]
+			+ SVerticalBox::Slot().FillHeight(1.f).Padding(2)
+			[
+				SAssignNew(TreeView, STreeView<TSharedPtr<FMaterialDiffTreeNode>>)
+					.SelectionMode(ESelectionMode::None)
+					.TreeItemsSource(&FilteredRootNodes)
+					.OnGenerateRow(this, &SMaterialInstanceDiffTree::OnGenerateRow)
+					.OnGetChildren_Lambda([](TSharedPtr<FMaterialDiffTreeNode> InNode, TArray<TSharedPtr<FMaterialDiffTreeNode>>& OutChildren)
+					{
+						OutChildren = InNode->Children;
+					})
+					.HeaderRow(
+						SNew(SHeaderRow)
+						+ SHeaderRow::Column(SMaterialInstanceDiffTreeRow::ColAssetPath)
+							.HAlignHeader(HAlign_Center)
+							.ManualWidth(GAssetNameColumnWidth_GroupDiff)
+							.ManualWidth_Lambda([this]()
+							{
+								return GAssetNameColumnWidth_GroupDiff;
+							})
+							.OnWidthChanged_Lambda([this](float NewWidth)
+							{
+								GAssetNameColumnWidth_GroupDiff = NewWidth;
+							})
+							.HeaderContent()
+							[
+								SNew(SBox)
+								.WidthOverride(TAttribute<FOptionalSize>::CreateLambda([]() -> FOptionalSize { return GAssetNameColumnWidth_GroupDiff; }))
+								[
+									SNew(SButton)
+									.ButtonStyle(FAppStyle::Get(), "NoBorder")
+									.HAlign(HAlign_Center)
+									.VAlign(VAlign_Center)
+									.ToolTipText(LOCTEXT("AssetPathHeaderTooltip", "Click to expand or collapse all rows"))
+									.OnClicked(this, &SMaterialInstanceDiffTree::OnToggleExpandAll)
+									[
+										SNew(STextBlock)
+										.Text(LOCTEXT("ColAssetPath", "Asset Path"))
+									]
+								]
+							]
+						+ SHeaderRow::Column(SMaterialInstanceDiffTreeRow::ColProperty)
+							.HAlignHeader(HAlign_Center)
+							.HAlignCell(HAlign_Center)
+							.DefaultLabel(LOCTEXT("ColProperty", "Property"))
+							.FixedWidth(150.f)
+						+ SHeaderRow::Column(SMaterialInstanceDiffTreeRow::ColOldValue)
+							.HAlignHeader(HAlign_Center)
+							.HAlignCell(HAlign_Center)
+							.DefaultLabel(LOCTEXT("ColOldValue", "Old Value"))
+							.FixedWidth(90.f)
+						+ SHeaderRow::Column(SMaterialInstanceDiffTreeRow::ColNewValue)
+							.HAlignHeader(HAlign_Center)
+							.HAlignCell(HAlign_Center)
+							.DefaultLabel(LOCTEXT("ColNewValue", "New Value"))
+							.FixedWidth(90.f)
+					)
+			]
+		];
+	}
+
+	/** Populate the tree for the given base material path. */
+	void SetMaterial(FSoftObjectPath const& InBaseMaterialPath)
+	{
+		SelectedMaterialPath = InBaseMaterialPath;
+		RebuildFromDiff();
+	}
+
+private:
+	/** Call GetMaterialValidationDescDiff for the current material and rebuild AllRootNodes from the result. */
+	void RebuildFromDiff()
+	{
+		AllRootNodes.Reset();
+
+		if (!OldGroup.IsValid() || !NewGroup.IsValid() || SelectedMaterialPath.IsNull())
+		{
+			ApplyFilters();
+			return;
+		}
+
+		TArray<FMaterialInstanceDiffResult> Diffs;
+		UMaterialValidationLibrary::GetMaterialValidationDescDiff(OldGroup.Get(), NewGroup.Get(), SelectedMaterialPath, /*bAllowAssetLoad=*/true, Diffs);
+
+		for (FMaterialInstanceDiffResult const& Diff : Diffs)
+		{
+			TSharedPtr<FMaterialDiffTreeNode> Node = MakeShared<FMaterialDiffTreeNode>();
+			Node->DiffResult = Diff;
+
+			if (Diff.DiffType == EMaterialInstanceDiffType::Modified)
+			{
+				for (FMaterialInstancePropertyDiff const& PropDiff : Diff.PropertyDiffs)
+				{
+					TSharedPtr<FMaterialDiffTreeNode> Child = MakeShared<FMaterialDiffTreeNode>();
+					Child->bIsPropertyNode = true;
+					Child->PropertyDiff = PropDiff;
+					Node->Children.Add(Child);
+				}
+			}
+
+			AllRootNodes.Add(MoveTemp(Node));
+		}
+
+		ApplyFilters();
+	}
+
+	/** Rebuild FilteredRootNodes from AllRootNodes by applying both the instance path and property name filters. */
+	void ApplyFilters()
+	{
+		FilteredRootNodes.Reset();
+
+		for (TSharedPtr<FMaterialDiffTreeNode> const& Node : AllRootNodes)
+		{
+			// Instance path filter.
+			if (!InstanceSearchText.IsEmpty() &&
+				!Node->DiffResult.InstancePath.ToString().Contains(InstanceSearchText, ESearchCase::IgnoreCase))
+			{
+				continue;
+			}
+
+			// Property name filter.
+			// Added/Removed nodes have no properties to match, so hide them when the filter is active. 
+			// Modified nodes are hidden unless a property name matches.
+			if (!PropertySearchText.IsEmpty())
+			{
+				if (Node->DiffResult.DiffType != EMaterialInstanceDiffType::Modified)
+				{
+					continue;
+				}
+
+				bool bHasAnyMatch = false;
+				for (TSharedPtr<FMaterialDiffTreeNode> const& Child : Node->Children)
+				{
+					if (Child->PropertyDiff.PropertyName.ToString().Contains(PropertySearchText, ESearchCase::IgnoreCase))
+					{
+						bHasAnyMatch = true;
+						break;
+					}
+				}
+				if (!bHasAnyMatch)
+				{
+					continue;
+				}
+			}
+
+			FilteredRootNodes.Add(Node);
+		}
+
+		if (TreeView.IsValid())
+		{
+			TreeView->RequestTreeRefresh();
+		}
+	}
+
+	/** Generate a table row for either a parent instance node or a child property diff node. */
+	TSharedRef<ITableRow> OnGenerateRow(TSharedPtr<FMaterialDiffTreeNode> InNode, const TSharedRef<STableViewBase>& InOwnerTable)
+	{
+		return SNew(SMaterialInstanceDiffTreeRow, InOwnerTable).Node(InNode);
+	}
+
+	/** Called when the instance path search box text changes. */
+	void OnInstanceSearchChanged(const FText& InText)
+	{
+		InstanceSearchText = InText.ToString();
+		ApplyFilters();
+	}
+
+	/** Called when the property name search box text changes. */
+	void OnPropertySearchChanged(const FText& InText)
+	{
+		PropertySearchText = InText.ToString();
+		ApplyFilters();
+	}
+
+	/** Called when the user clicks the Asset Path column header. Toggles between expanding and collapsing all tree nodes. */
+	FReply OnToggleExpandAll()
+	{
+		bAllExpanded = !bAllExpanded;
+		for (TSharedPtr<FMaterialDiffTreeNode> const& Node : AllRootNodes)
+		{
+			TreeView->SetItemExpansion(Node, bAllExpanded);
+		}
+		return FReply::Handled();
+	}
+};
+
+} // namespace MaterialValidation
+
+void SMaterialValidationGroupDiff::Construct(const FArguments& InArgs)
+{
+	OldGroup = InArgs._OldGroup;
+	NewGroup = InArgs._NewGroup;
+
+	TSharedPtr<MaterialValidation::SMaterialInstanceDiffTree> InstanceTree;
+
+	TArray<TSharedPtr<MaterialValidation::FMaterialDiffStats>> DiffRows;
+	if (OldGroup.IsValid() && NewGroup.IsValid())
+	{
+		DiffRows = MaterialValidation::BuildDiffRows(*OldGroup.Get(), *NewGroup.Get());
+	}
+
+	SAssignNew(InstanceTree, MaterialValidation::SMaterialInstanceDiffTree)
+		.OldGroup(OldGroup.Get())
+		.NewGroup(NewGroup.Get());
+
+	TWeakPtr<MaterialValidation::SMaterialInstanceDiffTree> WeakInstanceTree = InstanceTree;
+
+	ChildSlot
+	[
+		SNew(SSplitter)
+		.Orientation(Orient_Vertical)
+		+ SSplitter::Slot()
+		.Value(0.6f)
+		[
+			SNew(MaterialValidation::SMaterialDiffTable)
+				.InitialRows(DiffRows)
+				.OnRowSelected_Lambda([WeakInstanceTree](TSharedPtr<MaterialValidation::FMaterialDiffStats> InRow)
+				{
+					if (TSharedPtr<MaterialValidation::SMaterialInstanceDiffTree> Tree = WeakInstanceTree.Pin())
+					{
+						if (InRow.IsValid())
+						{
+							Tree->SetMaterial(InRow->AssetPath);
+						}
+					}
+				})
+		]
+		+ SSplitter::Slot()
+		.Value(0.4f)
+		[
+			InstanceTree.ToSharedRef()
+		]
+	];
+}
+
+TSharedPtr<SWindow> SMaterialValidationGroupDiff::CreateDiffWindow(UMaterialValidationGroup const* InOldGroup, UMaterialValidationGroup const* InNewGroup, FRevisionInfo const& InOldRevision, FRevisionInfo const& InNewRevision)
+{
+	auto MakeRevisionLabel = [](FRevisionInfo const& Rev) -> FString
+	{
+		if (!Rev.Revision.IsEmpty())
+		{
+			return Rev.Revision;
+		}
+		if (Rev.Changelist > 0)
+		{
+			return FString::FromInt(Rev.Changelist);
+		}
+		return TEXT("Unknown");
+	};
+
+	const FText WindowTitle = FText::Format(
+		LOCTEXT("DiffWindowTitle", "Material Validation Group Diff: {0} -> {1}"),
+		FText::FromString(MakeRevisionLabel(InOldRevision)),
+		FText::FromString(MakeRevisionLabel(InNewRevision)));
+
+	TSharedPtr<SWindow> Window = SNew(SWindow)
+		.Title(WindowTitle)
+		.ClientSize(FVector2f(1400.f, 900.f));
+
+	Window->SetContent(
+		SNew(SMaterialValidationGroupDiff)
+			.OldGroup(InOldGroup)
+			.NewGroup(InNewGroup)
+			.OldRevision(InOldRevision)
+			.NewRevision(InNewRevision));
+
+	const TSharedPtr<SWindow> ActiveModal = FSlateApplication::Get().GetActiveModalWindow();
+	if (ActiveModal.IsValid())
+	{
+		FSlateApplication::Get().AddWindowAsNativeChild(Window.ToSharedRef(), ActiveModal.ToSharedRef());
+	}
+	else
+	{
+		FSlateApplication::Get().AddWindow(Window.ToSharedRef());
+	}
+
+	return Window;
+}
+
+#undef LOCTEXT_NAMESPACE
diff --git a/Engine/Plugins/Editor/MaterialValidation/Source/Private/MaterialValidationGroupDiff.h b/Engine/Plugins/Editor/MaterialValidation/Source/Private/MaterialValidationGroupDiff.h
new file mode 100644
--- /dev/null
+++ b/Engine/Plugins/Editor/MaterialValidation/Source/Private/MaterialValidationGroupDiff.h
@@ -0,0 +1,34 @@
+// Copyright Epic Games, Inc. All Rights Reserved.
+
+#pragma once
+
+#include "AssetDefinition.h"
+#include "Widgets/SCompoundWidget.h"
+
+class UMaterialValidationGroup;
+
+/** Widget for the UMaterialValidationGroup asset diff window. */
+class SMaterialValidationGroupDiff : public SCompoundWidget
+{
+public:
+	SLATE_BEGIN_ARGS(SMaterialValidationGroupDiff) {}
+		SLATE_ARGUMENT(UMaterialValidationGroup const*, OldGroup)
+		SLATE_ARGUMENT(UMaterialValidationGroup const*, NewGroup)
+		SLATE_ARGUMENT(FRevisionInfo, OldRevision)
+		SLATE_ARGUMENT(FRevisionInfo, NewRevision)
+	SLATE_END_ARGS()
+
+	/** Build the widget layout from the provided group revisions. */
+	void Construct(const FArguments& InArgs);
+
+	/** Create and display a standalone diff window. */
+	static TSharedPtr<SWindow> CreateDiffWindow(
+		UMaterialValidationGroup const* InOldGroup,
+		UMaterialValidationGroup const* InNewGroup,
+		FRevisionInfo const& InOldRevision,
+		FRevisionInfo const& InNewRevision);
+
+private:
+	TWeakObjectPtr<const UMaterialValidationGroup> OldGroup;
+	TWeakObjectPtr<const UMaterialValidationGroup> NewGroup;
+};
diff --git a/Engine/Plugins/Editor/MaterialValidation/Source/Private/MaterialValidationLibrary.cpp b/Engine/Plugins/Editor/MaterialValidation/Source/Private/MaterialValidationLibrary.cpp
--- a/Engine/Plugins/Editor/MaterialValidation/Source/Private/MaterialValidationLibrary.cpp
+++ b/Engine/Plugins/Editor/MaterialValidation/Source/Private/MaterialValidationLibrary.cpp
@@ -2,6 +2,7 @@
 
 #include "MaterialValidationLibrary.h"
 
+#include "Algo/Sort.h"
 #include "AssetRegistry/AssetData.h"
 #include "AssetRegistry/AssetRegistryModule.h"
 #include "AssetRegistry/IAssetRegistry.h"
@@ -12,6 +13,7 @@
 #include "Materials/MaterialLayersFunctions.h"
 #include "MaterialValidationConfig.h"
 #include "MaterialValidationGroup.h"
+#include "MaterialValidationLibraryTypes.h"
 #include "MaterialValidationModule.h"
 #include "Misc/ScopedSlowTask.h"
 
@@ -1428,7 +1430,6 @@ namespace MaterialValidationLibrary {
  */
 enum class EMaterialInstanceStaticProperties : int32
 {
-	PermutationHash,
 	TwoSided,
 	bIsThinSurface,
 	DitheredLODTransition,
@@ -1445,79 +1446,75 @@ enum class EMaterialInstanceStaticProperties : int32
 /** Helper getter functions for each of the EMaterialInstanceStaticProperties. */
 struct FMaterialInstanceStaticPropertyGetters
 {
-	static uint32 GetPermutationHash(FMaterialValidationAssetData_MaterialInstance const& P) { return P.PermutationHash; }
-	static uint8 GetTwoSided(FMaterialValidationAssetData_MaterialInstance const& P) { return P.StaticProperties.TwoSided; }
-	static uint8 GetIsThinSurface(FMaterialValidationAssetData_MaterialInstance const& P) { return P.StaticProperties.bIsThinSurface; }
-	static uint8 GetDitheredLODTransition(FMaterialValidationAssetData_MaterialInstance const& P) { return P.StaticProperties.DitheredLODTransition; }
-	static uint8 GetCastDynamicShadowAsMasked(FMaterialValidationAssetData_MaterialInstance const& P) { return P.StaticProperties.bCastDynamicShadowAsMasked; }
-	static uint8 GetOutputTranslucentVelocity(FMaterialValidationAssetData_MaterialInstance const& P) { return P.StaticProperties.bOutputTranslucentVelocity; }
-	static uint8 GetHasPixelAnimation(FMaterialValidationAssetData_MaterialInstance const& P) { return P.StaticProperties.bHasPixelAnimation; }
-	static uint8 GetEnableTessellation(FMaterialValidationAssetData_MaterialInstance const& P) { return P.StaticProperties.bEnableTessellation; }
-	static EBlendMode GetBlendMode(FMaterialValidationAssetData_MaterialInstance const& P) { return P.StaticProperties.BlendMode; }
-	static EMaterialShadingModel GetShadingModel(FMaterialValidationAssetData_MaterialInstance const& P) { return P.StaticProperties.ShadingModel; }
-	static float GetOpacityMaskClipValue(FMaterialValidationAssetData_MaterialInstance const& P) { return P.StaticProperties.OpacityMaskClipValue; }
+	static uint8 GetTwoSided(FStaticPermutationProperties const& P) { return P.TwoSided; }
+	static uint8 GetIsThinSurface(FStaticPermutationProperties const& P) { return P.bIsThinSurface; }
+	static uint8 GetDitheredLODTransition(FStaticPermutationProperties const& P) { return P.DitheredLODTransition; }
+	static uint8 GetCastDynamicShadowAsMasked(FStaticPermutationProperties const& P) { return P.bCastDynamicShadowAsMasked; }
+	static uint8 GetOutputTranslucentVelocity(FStaticPermutationProperties const& P) { return P.bOutputTranslucentVelocity; }
+	static uint8 GetHasPixelAnimation(FStaticPermutationProperties const& P) { return P.bHasPixelAnimation; }
+	static uint8 GetEnableTessellation(FStaticPermutationProperties const& P) { return P.bEnableTessellation; }
+	static EBlendMode GetBlendMode(FStaticPermutationProperties const& P) { return P.BlendMode; }
+	static EMaterialShadingModel GetShadingModel(FStaticPermutationProperties const& P) { return P.ShadingModel; }
+	static float GetOpacityMaskClipValue(FStaticPermutationProperties const& P) { return P.OpacityMaskClipValue; }
 };
 
 struct FMaterialInstanceStaticPropertyModifiedGetters
 {
-	static bool GetPermutationHashModified(FMaterialValidationAssetData_MaterialInstance const& P) { return false; }
-	static bool GetTwoSidedModified(FMaterialValidationAssetData_MaterialInstance const& P) { return P.StaticPropertyOverrideFlags.bOverride_TwoSided; }
-	static bool GetIsThinSurfaceModified(FMaterialValidationAssetData_MaterialInstance const& P) { return P.StaticPropertyOverrideFlags.bOverride_bIsThinSurface; }
-	static bool GetDitheredLODTransitionModified(FMaterialValidationAssetData_MaterialInstance const& P) { return P.StaticPropertyOverrideFlags.bOverride_DitheredLODTransition; }
-	static bool GetCastDynamicShadowAsMaskedModified(FMaterialValidationAssetData_MaterialInstance const& P) { return P.StaticPropertyOverrideFlags.bOverride_CastDynamicShadowAsMasked; }
-	static bool GetOutputTranslucentVelocityModified(FMaterialValidationAssetData_MaterialInstance const& P) { return P.StaticPropertyOverrideFlags.bOverride_OutputTranslucentVelocity; }
-	static bool GetHasPixelAnimationModified(FMaterialValidationAssetData_MaterialInstance const& P) { return P.StaticPropertyOverrideFlags.bOverride_bHasPixelAnimation; }
-	static bool GetEnableTessellationModified(FMaterialValidationAssetData_MaterialInstance const& P) { return P.StaticPropertyOverrideFlags.bOverride_bEnableTessellation; }
-	static bool GetBlendModeModified(FMaterialValidationAssetData_MaterialInstance const& P) { return P.StaticPropertyOverrideFlags.bOverride_BlendMode; }
-	static bool GetShadingModelModified(FMaterialValidationAssetData_MaterialInstance const& P) { return P.StaticPropertyOverrideFlags.bOverride_ShadingModel; }
-	static bool GetOpacityMaskClipModified(FMaterialValidationAssetData_MaterialInstance const& P) { return P.StaticPropertyOverrideFlags.bOverride_OpacityMaskClipValue; }
+	static bool GetTwoSidedModified(FStaticPermutationPropertyOverrideFlags const& P) { return P.bOverride_TwoSided; }
+	static bool GetIsThinSurfaceModified(FStaticPermutationPropertyOverrideFlags const& P) { return P.bOverride_bIsThinSurface; }
+	static bool GetDitheredLODTransitionModified(FStaticPermutationPropertyOverrideFlags const& P) { return P.bOverride_DitheredLODTransition; }
+	static bool GetCastDynamicShadowAsMaskedModified(FStaticPermutationPropertyOverrideFlags const& P) { return P.bOverride_CastDynamicShadowAsMasked; }
+	static bool GetOutputTranslucentVelocityModified(FStaticPermutationPropertyOverrideFlags const& P) { return P.bOverride_OutputTranslucentVelocity; }
+	static bool GetHasPixelAnimationModified(FStaticPermutationPropertyOverrideFlags const& P) { return P.bOverride_bHasPixelAnimation; }
+	static bool GetEnableTessellationModified(FStaticPermutationPropertyOverrideFlags const& P) { return P.bOverride_bEnableTessellation; }
+	static bool GetBlendModeModified(FStaticPermutationPropertyOverrideFlags const& P) { return P.bOverride_BlendMode; }
+	static bool GetShadingModelModified(FStaticPermutationPropertyOverrideFlags const& P) { return P.bOverride_ShadingModel; }
+	static bool GetOpacityMaskClipModified(FStaticPermutationPropertyOverrideFlags const& P) { return P.bOverride_OpacityMaskClipValue; }
 };
 
 template <auto GetterPtr>
-static FString GetStaticPropertyValueAsString(FMaterialValidationAssetData_MaterialInstance const& AssetData)
+static FString GetStaticPropertyValueAsString(FStaticPermutationProperties const& InProperties)
 {
-	using ReturnType = TInvokeResult_T<decltype(GetterPtr), FMaterialValidationAssetData_MaterialInstance const&>;
+	using ReturnType = TInvokeResult_T<decltype(GetterPtr), FStaticPermutationProperties const&>;
 
 	if constexpr (TIsEnum<ReturnType>::Value)
 	{
-		return  UEnum::GetDisplayValueAsText((*GetterPtr)(AssetData)).ToString();
+		return  UEnum::GetDisplayValueAsText((*GetterPtr)(InProperties)).ToString();
 	}
  	else if constexpr (TIsFloatingPoint<ReturnType>::Value)
  	{
- 		return FString::Printf(TEXT("%g"), (*GetterPtr)(AssetData));
+ 		return FString::Printf(TEXT("%g"), (*GetterPtr)(InProperties));
  	}
 	else
 	{
-		return LexToString((*GetterPtr)(AssetData));
+		return LexToString((*GetterPtr)(InProperties));
 	}
 }
 
-using FnGetPropertyValueAsString = FString(*)(FMaterialValidationAssetData_MaterialInstance const&);
-using FnIsPropertyModified = bool(*)(FMaterialValidationAssetData_MaterialInstance const&);
+using FnGetPropertyValueAsString = FString(*)(FStaticPermutationProperties const&);
+using FnIsPropertyModified = bool(*)(FStaticPermutationPropertyOverrideFlags const&);
 
 /** Struct containing information needed to generate table entries for EMaterialInstanceStaticProperties. */
 struct FMaterialInstanceStaticPropertyInfo
 {
 	FName PropertyId;
 	FName PropertyName;
-	EMaterialInstanceColumnType Category;
 	FnGetPropertyValueAsString GetPropertyValueAsStringFunction;
 	FnIsPropertyModified GetIsPropertyModifiedFunction;
 };
 
 /** Array for table generation data for EMaterialInstanceStaticProperties. */
 static const FMaterialInstanceStaticPropertyInfo StaticPropertyInfos[(int32)EMaterialInstanceStaticProperties::Count] = {
-	{"PID_PermutationHash", "Permutation Hash", EMaterialInstanceColumnType::Default, &GetStaticPropertyValueAsString<&FMaterialInstanceStaticPropertyGetters::GetPermutationHash>, &FMaterialInstanceStaticPropertyModifiedGetters::GetPermutationHashModified },
-	{"PID_TwoSided", "Two Sided", EMaterialInstanceColumnType::StaticProperty, &GetStaticPropertyValueAsString <&FMaterialInstanceStaticPropertyGetters::GetTwoSided>, &FMaterialInstanceStaticPropertyModifiedGetters::GetTwoSidedModified },
-	{"PID_IsThinSurface", "Thin Surface", EMaterialInstanceColumnType::StaticProperty, &GetStaticPropertyValueAsString <&FMaterialInstanceStaticPropertyGetters::GetIsThinSurface>, &FMaterialInstanceStaticPropertyModifiedGetters::GetIsThinSurfaceModified },
-	{"PID_DitheredLODTransition", "Dithered LOD Transition", EMaterialInstanceColumnType::StaticProperty, &GetStaticPropertyValueAsString <&FMaterialInstanceStaticPropertyGetters::GetDitheredLODTransition>, &FMaterialInstanceStaticPropertyModifiedGetters::GetDitheredLODTransitionModified },
-	{"PID_CastDynamicShadowAsMasked", "Dynamic Shadow As Masked", EMaterialInstanceColumnType::StaticProperty, &GetStaticPropertyValueAsString <&FMaterialInstanceStaticPropertyGetters::GetCastDynamicShadowAsMasked>, &FMaterialInstanceStaticPropertyModifiedGetters::GetCastDynamicShadowAsMaskedModified },
-	{"PID_OutputTranslucentVelocity", "Output Translucent Velocity", EMaterialInstanceColumnType::StaticProperty, &GetStaticPropertyValueAsString <&FMaterialInstanceStaticPropertyGetters::GetOutputTranslucentVelocity>, &FMaterialInstanceStaticPropertyModifiedGetters::GetOutputTranslucentVelocityModified },
-	{"PID_HasPixelAnimation", "Has Pixel Animation", EMaterialInstanceColumnType::StaticProperty, &GetStaticPropertyValueAsString <&FMaterialInstanceStaticPropertyGetters::GetHasPixelAnimation>, &FMaterialInstanceStaticPropertyModifiedGetters::GetHasPixelAnimationModified },
-	{"PID_EnableTessellation", "Tessellation", EMaterialInstanceColumnType::StaticProperty, &GetStaticPropertyValueAsString <&FMaterialInstanceStaticPropertyGetters::GetEnableTessellation>, &FMaterialInstanceStaticPropertyModifiedGetters::GetEnableTessellationModified },
-	{"PID_BlendMode", "Blend Mode", EMaterialInstanceColumnType::StaticProperty, &GetStaticPropertyValueAsString <&FMaterialInstanceStaticPropertyGetters::GetBlendMode>, &FMaterialInstanceStaticPropertyModifiedGetters::GetBlendModeModified },
-	{"PID_ShadingModel", "Shading Model", EMaterialInstanceColumnType::StaticProperty, &GetStaticPropertyValueAsString <&FMaterialInstanceStaticPropertyGetters::GetShadingModel>, &FMaterialInstanceStaticPropertyModifiedGetters::GetShadingModelModified },
-	{"PID_OpacityMaskClipValue", "Opacity Mask Clip Value", EMaterialInstanceColumnType::StaticProperty, &GetStaticPropertyValueAsString <&FMaterialInstanceStaticPropertyGetters::GetOpacityMaskClipValue>, &FMaterialInstanceStaticPropertyModifiedGetters::GetOpacityMaskClipModified },
+	{"PID_TwoSided", "Two Sided", &GetStaticPropertyValueAsString <&FMaterialInstanceStaticPropertyGetters::GetTwoSided>, &FMaterialInstanceStaticPropertyModifiedGetters::GetTwoSidedModified },
+	{"PID_IsThinSurface", "Thin Surface", &GetStaticPropertyValueAsString <&FMaterialInstanceStaticPropertyGetters::GetIsThinSurface>, &FMaterialInstanceStaticPropertyModifiedGetters::GetIsThinSurfaceModified },
+	{"PID_DitheredLODTransition", "Dithered LOD Transition", &GetStaticPropertyValueAsString <&FMaterialInstanceStaticPropertyGetters::GetDitheredLODTransition>, &FMaterialInstanceStaticPropertyModifiedGetters::GetDitheredLODTransitionModified },
+	{"PID_CastDynamicShadowAsMasked", "Dynamic Shadow As Masked", &GetStaticPropertyValueAsString <&FMaterialInstanceStaticPropertyGetters::GetCastDynamicShadowAsMasked>, &FMaterialInstanceStaticPropertyModifiedGetters::GetCastDynamicShadowAsMaskedModified },
+	{"PID_OutputTranslucentVelocity", "Output Translucent Velocity", &GetStaticPropertyValueAsString <&FMaterialInstanceStaticPropertyGetters::GetOutputTranslucentVelocity>, &FMaterialInstanceStaticPropertyModifiedGetters::GetOutputTranslucentVelocityModified },
+	{"PID_HasPixelAnimation", "Has Pixel Animation", &GetStaticPropertyValueAsString <&FMaterialInstanceStaticPropertyGetters::GetHasPixelAnimation>, &FMaterialInstanceStaticPropertyModifiedGetters::GetHasPixelAnimationModified },
+	{"PID_EnableTessellation", "Tessellation", &GetStaticPropertyValueAsString <&FMaterialInstanceStaticPropertyGetters::GetEnableTessellation>, &FMaterialInstanceStaticPropertyModifiedGetters::GetEnableTessellationModified },
+	{"PID_BlendMode", "Blend Mode", &GetStaticPropertyValueAsString <&FMaterialInstanceStaticPropertyGetters::GetBlendMode>, &FMaterialInstanceStaticPropertyModifiedGetters::GetBlendModeModified },
+	{"PID_ShadingModel", "Shading Model", &GetStaticPropertyValueAsString <&FMaterialInstanceStaticPropertyGetters::GetShadingModel>, &FMaterialInstanceStaticPropertyModifiedGetters::GetShadingModelModified },
+	{"PID_OpacityMaskClipValue", "Opacity Mask Clip Value", &GetStaticPropertyValueAsString <&FMaterialInstanceStaticPropertyGetters::GetOpacityMaskClipValue>, &FMaterialInstanceStaticPropertyModifiedGetters::GetOpacityMaskClipModified },
 };
 
 /** Static array of column FNames for all material usages. */
@@ -1555,6 +1552,246 @@ static int32 GetPropertyCount(FMaterialDatabaseAssetHierarchyInfo const& InInfo)
 	return UE_ARRAY_COUNT(MaterialValidationLibrary::StaticPropertyInfos) + MATUSAGE_MAX + InInfo.StaticSwitchNames.Num() + InInfo.ComponentMaskNames.Num();
 }
 
+/** Read a single bit from a packed uint32 array (1 bit per element). Returns false if the index is out of range. */
+static bool GetPackedBit(TConstArrayView<uint32> InArray, int32 Index)
+{
+	const uint32 BitArrayIndex = (uint32)Index / 32;
+	const uint32 BitIndex = (uint32)Index % 32;
+	return InArray.IsValidIndex(BitArrayIndex) && (InArray[BitArrayIndex] & (1u << BitIndex)) != 0;
+}
+
+/** Read a nibble (4 bits) from a packed uint32 array (4 bits per element). Returns 0 if the index is out of range. */
+static uint32 GetPackedNibble(TConstArrayView<uint32> InArray, int32 Index)
+{
+	const uint32 BitArrayIndex = Index / 8;
+	const uint32 BitIndex = (Index * 4) % 32;
+	return InArray.IsValidIndex(BitArrayIndex) ? (InArray[BitArrayIndex] >> BitIndex) & 0xf : 0u;
+}
+
+/** Format a component mask nibble as an RGBA_ string. */
+static FString FormatComponentMask(uint32 InValue)
+{
+	TStringBuilder<8> RGBA;
+	RGBA.AppendChar(InValue & 1 ? 'R' : '_');
+	RGBA.AppendChar(InValue & 2 ? 'G' : '_');
+	RGBA.AppendChar(InValue & 4 ? 'B' : '_');
+	RGBA.AppendChar(InValue & 8 ? 'A' : '_');
+	return *RGBA;
+}
+
+/** String shown to indicate that a material instance property is being inherited from its parent. */
+static const TCHAR* GInheritedValueString = TEXT("[Inherited]");
+
+/** Add an entry to InOutDiffs if old and new string values differ. */
+static void AddDiffIfChanged(FName InName, EMaterialInstanceColumnType InCategory, FString const& InOldValue, FString const& InNewValue, TArray<FMaterialInstancePropertyDiff>& InOutDiffs)
+{
+	if (InOldValue != InNewValue)
+	{
+		InOutDiffs.Add({ .PropertyName = InName, .Category = InCategory, .OldValue = InOldValue, .NewValue = InNewValue });
+	}
+}
+
+/** Diff all static permutation properties between two base materials (no inheritance). */
+static void BuildStaticPropertyDiff(
+	FStaticPermutationProperties const& InOld,
+	FStaticPermutationProperties const& InNew,
+	TArray<FMaterialInstancePropertyDiff>& InOutDiffs)
+{
+	for (MaterialValidationLibrary::FMaterialInstanceStaticPropertyInfo const& StaticPropertyInfo : MaterialValidationLibrary::StaticPropertyInfos)
+	{
+		AddDiffIfChanged(
+			StaticPropertyInfo.PropertyName,
+			EMaterialInstanceColumnType::StaticProperty,
+			StaticPropertyInfo.GetPropertyValueAsStringFunction(InOld),
+			StaticPropertyInfo.GetPropertyValueAsStringFunction(InNew),
+			InOutDiffs);
+	}
+}
+
+/** Diff all static permutation properties between two material instances, showing "[Inherited]" for non-overridden properties. */
+static void BuildStaticPropertyDiff(
+	FStaticPermutationProperties const& InOld,
+	FStaticPermutationProperties const& InNew,
+	FStaticPermutationPropertyOverrideFlags const& InOldOverrideFlags,
+	FStaticPermutationPropertyOverrideFlags const& InNewOverrideFlags,
+	TArray<FMaterialInstancePropertyDiff>& InOutDiffs)
+{
+	for (MaterialValidationLibrary::FMaterialInstanceStaticPropertyInfo const& StaticPropertyInfo : MaterialValidationLibrary::StaticPropertyInfos)
+	{
+		const bool bOldInherited = !StaticPropertyInfo.GetIsPropertyModifiedFunction(InOldOverrideFlags);
+		const bool bNewInherited = !StaticPropertyInfo.GetIsPropertyModifiedFunction(InNewOverrideFlags);
+		AddDiffIfChanged(
+			StaticPropertyInfo.PropertyName,
+			EMaterialInstanceColumnType::StaticProperty,
+			bOldInherited ? GInheritedValueString : StaticPropertyInfo.GetPropertyValueAsStringFunction(InOld),
+			bNewInherited ? GInheritedValueString : StaticPropertyInfo.GetPropertyValueAsStringFunction(InNew),
+			InOutDiffs);
+	}
+}
+
+/** Diff all usage flags between two base materials (no inheritance). */
+static void BuildUsageFlagDiff(
+	FStaticPermutationProperties const& InOld,
+	FStaticPermutationProperties const& InNew,
+	TArray<FMaterialInstancePropertyDiff>& InOutDiffs)
+{
+	for (int32 Index = 0; Index < MATUSAGE_MAX; ++Index)
+	{
+		const FName UsageName(*UMaterialInterface::GetUsageName((EMaterialUsage)Index));
+		const FString OldValue = (InOld.UsageFlags & (1u << Index)) != 0 ? TEXT("1") : TEXT("0");
+		const FString NewValue = (InNew.UsageFlags & (1u << Index)) != 0 ? TEXT("1") : TEXT("0");
+		AddDiffIfChanged(UsageName, EMaterialInstanceColumnType::UsageFlag, OldValue, NewValue, InOutDiffs);
+	}
+}
+
+/** Diff all usage flags between two material instances, showing "[Inherited]" for non-overridden flags. */
+static void BuildUsageFlagDiff(
+	FStaticPermutationProperties const& InOld,
+	FStaticPermutationProperties const& InNew,
+	FStaticPermutationPropertyOverrideFlags const& InOldOverrideFlags,
+	FStaticPermutationPropertyOverrideFlags const& InNewOverrideFlags,
+	TArray<FMaterialInstancePropertyDiff>& InOutDiffs)
+{
+	for (int32 Index = 0; Index < MATUSAGE_MAX; ++Index)
+	{
+		const FName UsageName(*UMaterialInterface::GetUsageName((EMaterialUsage)Index));
+		const FString OldValue = !(InOldOverrideFlags.Override_UsageFlags & (1u << Index)) ? GInheritedValueString : (InOld.UsageFlags & (1u << Index)) != 0 ? TEXT("1") : TEXT("0");
+		const FString NewValue = !(InNewOverrideFlags.Override_UsageFlags & (1u << Index)) ? GInheritedValueString : (InNew.UsageFlags & (1u << Index)) != 0 ? TEXT("1") : TEXT("0");
+		AddDiffIfChanged(UsageName, EMaterialInstanceColumnType::UsageFlag, OldValue, NewValue, InOutDiffs);
+	}
+}
+
+/** Diff all static switch parameters between two base materials (no inheritance). */
+static void BuildStaticSwitchDiff(
+	TConstArrayView<uint32> InOld,
+	TConstArrayView<uint32> InNew,
+	TConstArrayView<FName> InStaticSwitchNames,
+	TArray<FMaterialInstancePropertyDiff>& InOutDiffs)
+{
+	for (int32 Index = 0; Index < InStaticSwitchNames.Num(); ++Index)
+	{
+		const FString OldValue = GetPackedBit(InOld, Index) ? TEXT("1") : TEXT("0");
+		const FString NewValue = GetPackedBit(InNew, Index) ? TEXT("1") : TEXT("0");
+		AddDiffIfChanged(InStaticSwitchNames[Index], EMaterialInstanceColumnType::StaticSwitch, OldValue, NewValue, InOutDiffs);
+	}
+}
+
+/** Diff all static switch parameters between two material instances, showing "[Inherited]" for non-overridden switches. */
+static void BuildStaticSwitchDiff(
+	TConstArrayView<uint32> InOld,
+	TConstArrayView<uint32> InNew,
+	TConstArrayView<uint32> InOldOverrideMask,
+	TConstArrayView<uint32> InNewOverrideMask,
+	TConstArrayView<FName> InStaticSwitchNames,
+	TArray<FMaterialInstancePropertyDiff>& InOutDiffs)
+{
+	for (int32 Index = 0; Index < InStaticSwitchNames.Num(); ++Index)
+	{
+		const FString OldValue = !GetPackedBit(InOldOverrideMask, Index) ? GInheritedValueString : GetPackedBit(InOld, Index) ? TEXT("1") : TEXT("0");
+		const FString NewValue = !GetPackedBit(InNewOverrideMask, Index) ? GInheritedValueString : GetPackedBit(InNew, Index) ? TEXT("1") : TEXT("0");
+		AddDiffIfChanged(InStaticSwitchNames[Index], EMaterialInstanceColumnType::StaticSwitch, OldValue, NewValue, InOutDiffs);
+	}
+}
+
+/** Diff all component mask parameters between two base materials (no inheritance). */
+static void BuildComponentMaskDiff(
+	TConstArrayView<uint32> InOld,
+	TConstArrayView<uint32> InNew,
+	TConstArrayView<FName> InComponentMaskNames,
+	TArray<FMaterialInstancePropertyDiff>& InOutDiffs)
+{
+	for (int32 Index = 0; Index < InComponentMaskNames.Num(); ++Index)
+	{
+		const FString OldValue = FormatComponentMask(GetPackedNibble(InOld, Index));
+		const FString NewValue = FormatComponentMask(GetPackedNibble(InNew, Index));
+		AddDiffIfChanged(InComponentMaskNames[Index], EMaterialInstanceColumnType::ComponentMask, OldValue, NewValue, InOutDiffs);
+	}
+}
+
+/** Diff all component mask parameters between two material instances, showing "[Inherited]" for non-overridden masks. */
+static void BuildComponentMaskDiff(
+	TConstArrayView<uint32> InOld,
+	TConstArrayView<uint32> InNew,
+	TConstArrayView<uint32> InOldOverrideMask,
+	TConstArrayView<uint32> InNewOverrideMask,
+	TConstArrayView<FName> InComponentMaskNames,
+	TArray<FMaterialInstancePropertyDiff>& InOutDiffs)
+{
+	for (int32 Index = 0; Index < InComponentMaskNames.Num(); ++Index)
+	{
+		const FString OldValue = !GetPackedNibble(InOldOverrideMask, Index) ? GInheritedValueString : FormatComponentMask(GetPackedNibble(InOld, Index));
+		const FString NewValue = !GetPackedNibble(InNewOverrideMask, Index) ? GInheritedValueString : FormatComponentMask(GetPackedNibble(InNew, Index));
+		AddDiffIfChanged(InComponentMaskNames[Index], EMaterialInstanceColumnType::ComponentMask, OldValue, NewValue, InOutDiffs);
+	}
+}
+
+/** Build property diffs for a base material entry. */
+static void BuildBaseMaterialPropertyDiff(
+	FMaterialValidationAssetData_Material const& InOldData,
+	FMaterialValidationAssetData_Material const& InNewData,
+	TConstArrayView<FName> InStaticSwitchNames,
+	TConstArrayView<FName> InComponentMaskNames,
+	TArray<FMaterialInstancePropertyDiff>& OutDiffs)
+{
+	BuildStaticPropertyDiff(InOldData.StaticProperties, InNewData.StaticProperties, OutDiffs);
+	BuildUsageFlagDiff(InOldData.StaticProperties, InNewData.StaticProperties, OutDiffs);
+
+	// Switch parameters are only comparable when layout hashes match.
+	if (InOldData.StaticPropertyLayoutHash == InNewData.StaticPropertyLayoutHash)
+	{
+		BuildStaticSwitchDiff(InOldData.StaticSwitchValues, InNewData.StaticSwitchValues, InStaticSwitchNames, OutDiffs);
+		BuildComponentMaskDiff(InOldData.ComponentMaskValues, InNewData.ComponentMaskValues, InComponentMaskNames, OutDiffs);
+	}
+	else
+	{
+		OutDiffs.Add({
+			.PropertyName = "Static Parameter Layout",
+			.OldValue = FString::Printf(TEXT("0x%08X"), InOldData.StaticPropertyLayoutHash),
+			.NewValue = FString::Printf(TEXT("0x%08X"), InNewData.StaticPropertyLayoutHash) });
+	}
+
+	// Material layer hash.
+	if (InOldData.MaterialLayerHash != InNewData.MaterialLayerHash)
+	{
+		OutDiffs.Add({
+			.PropertyName = "Material Layer",
+			.OldValue = FString::Printf(TEXT("0x%08X"), InOldData.MaterialLayerHash),
+			.NewValue = FString::Printf(TEXT("0x%08X"), InNewData.MaterialLayerHash) });
+	}
+}
+
+/**
+ * Build property diffs for a material instance entry.
+ * Compares raw override data and returns a value of "[Inherited]" when not overriding.
+ */
+static void BuildInstancePropertyDiff(
+	FMaterialValidationAssetData_MaterialInstance const& InOldData,
+	FMaterialValidationAssetData_MaterialInstance const& InNewData,
+	TConstArrayView<FName> InStaticSwitchNames,
+	TConstArrayView<FName> InComponentMaskNames,
+	bool bLayoutsMatch,
+	TArray<FMaterialInstancePropertyDiff>& OutDiffs)
+{
+	// Static permutation properties and usage flags with override flags.
+	BuildStaticPropertyDiff(InOldData.StaticProperties, InNewData.StaticProperties, InOldData.StaticPropertyOverrideFlags, InNewData.StaticPropertyOverrideFlags, OutDiffs);
+	BuildUsageFlagDiff(InOldData.StaticProperties, InNewData.StaticProperties, InOldData.StaticPropertyOverrideFlags, InNewData.StaticPropertyOverrideFlags, OutDiffs);
+
+	// Switch parameters are only comparable when layout hashes match.
+	if (bLayoutsMatch)
+	{
+		BuildStaticSwitchDiff(InOldData.StaticSwitchOverrideValues, InNewData.StaticSwitchOverrideValues, InOldData.StaticSwitchOverrideMask, InNewData.StaticSwitchOverrideMask, InStaticSwitchNames, OutDiffs);
+		BuildComponentMaskDiff(InOldData.ComponentMaskOverrideValues, InNewData.ComponentMaskOverrideValues, InOldData.ComponentMaskOverrideMask, InNewData.ComponentMaskOverrideMask, InComponentMaskNames, OutDiffs);
+	}
+
+	// Material layer hash.
+	if (InOldData.MaterialLayerHash != InNewData.MaterialLayerHash)
+	{
+		OutDiffs.Add({ .PropertyName = "Material Layer",
+			.OldValue = FString::Printf(TEXT("0x%08X"), InOldData.MaterialLayerHash),
+			.NewValue = FString::Printf(TEXT("0x%08X"), InNewData.MaterialLayerHash) });
+	}
+}
+
 } // MaterialValidationLibrary
 
 void UMaterialValidationLibrary::GetMaterialProperties(FMaterialDatabaseAssetHierarchyInfo const& InInfo, TArray<FMaterialDatabaseAssetPropertyDesc>& OutProperties)
@@ -1564,7 +1801,7 @@ void UMaterialValidationLibrary::GetMaterialProperties(FMaterialDatabaseAssetHie
 
 	for (MaterialValidationLibrary::FMaterialInstanceStaticPropertyInfo const& StaticPropertyInfo : MaterialValidationLibrary::StaticPropertyInfos)
 	{
-		OutProperties.Add({ .Id = StaticPropertyInfo.PropertyId, .Name = StaticPropertyInfo.PropertyName, .Category = StaticPropertyInfo.Category });
+		OutProperties.Add({ .Id = StaticPropertyInfo.PropertyId, .Name = StaticPropertyInfo.PropertyName, .Category = EMaterialInstanceColumnType::StaticProperty });
 	}
 	for (int32 Index = 0; Index < MATUSAGE_MAX; ++Index)
 	{
@@ -1595,8 +1832,8 @@ void UMaterialValidationLibrary::GetMaterialPropertyValues(FMaterialDatabaseAsse
 	for (MaterialValidationLibrary::FMaterialInstanceStaticPropertyInfo const& StaticPropertyInfo : MaterialValidationLibrary::StaticPropertyInfos)
 	{
 		OutValues.Add({
-			.Value = StaticPropertyInfo.GetPropertyValueAsStringFunction(AssetData), 
-			.bIsModifiedFromParent = StaticPropertyInfo.GetIsPropertyModifiedFunction(AssetData) });
+			.Value = StaticPropertyInfo.GetPropertyValueAsStringFunction(AssetData.StaticProperties),
+			.bIsModifiedFromParent = StaticPropertyInfo.GetIsPropertyModifiedFunction(AssetData.StaticPropertyOverrideFlags) });
 	}
 	
 	for (int32 Index = 0; Index < MATUSAGE_MAX; ++Index)
@@ -1633,12 +1870,7 @@ void UMaterialValidationLibrary::GetMaterialPropertyValues(FMaterialDatabaseAsse
 		{
 			const uint32 Value = (AssetData.ComponentMaskOverrideValues[BitArrayIndex] >> BitIndex) & 0xf;
 			const bool bModified = ((AssetData.ComponentMaskOverrideMask[BitArrayIndex] >> BitIndex) & 0xf) != 0;
-			TStringBuilder<8> RGBA;
-			RGBA.AppendChar(Value & 1 ? 'R' : '_');
-			RGBA.AppendChar(Value & 2 ? 'G' : '_');
-			RGBA.AppendChar(Value & 4 ? 'B' : '_');
-			RGBA.AppendChar(Value & 8 ? 'A' : '_');
-			OutValues.Add({ .Value = *RGBA, .bIsModifiedFromParent = bModified });
+			OutValues.Add({ .Value = MaterialValidationLibrary::FormatComponentMask(Value), .bIsModifiedFromParent = bModified});
 		}
 		else
 		{
@@ -1648,4 +1880,163 @@ void UMaterialValidationLibrary::GetMaterialPropertyValues(FMaterialDatabaseAsse
 	}
 }
 
+void UMaterialValidationLibrary::GetMaterialValidationDescDiff(
+	UMaterialValidationGroup const* InGroupOld,
+	UMaterialValidationGroup const* InGroupNew,
+	FSoftObjectPath const& InBaseMaterialPath,
+	bool bAllowAssetLoad,
+	TArray<FMaterialInstanceDiffResult>& OutDiffs)
+{
+	OutDiffs.Reset();
+
+	FMaterialValidationDesc const* OldDesc = InGroupOld != nullptr ? MaterialValidationLibrary::FindMaterialDescFromSoftObjectPath(*InGroupOld, InBaseMaterialPath) : nullptr;
+	FMaterialValidationDesc const* NewDesc = InGroupNew != nullptr ? MaterialValidationLibrary::FindMaterialDescFromSoftObjectPath(*InGroupNew, InBaseMaterialPath) : nullptr;
+
+	if (OldDesc == nullptr && NewDesc == nullptr)
+	{
+		return;
+	}
+	else if (OldDesc == nullptr && NewDesc != nullptr)
+	{
+		OutDiffs.Add({ .InstancePath = InBaseMaterialPath, .DiffType = EMaterialInstanceDiffType::Added });
+		return;
+	}
+	else if (OldDesc != nullptr && NewDesc == nullptr)
+	{
+		OutDiffs.Add({ .InstancePath = InBaseMaterialPath, .DiffType = EMaterialInstanceDiffType::Removed });
+		return;
+	}
+
+	// Resolve switch and component mask parameter names if allowed.
+	TArray<FName> SwitchNames;
+	TArray<FName> ComponentMaskNames;
+
+	TSoftObjectPtr<UMaterial> BaseMaterialSoftPtr(InBaseMaterialPath);
+	bool bFoundParameterNames = false;
+	if (bAllowAssetLoad || BaseMaterialSoftPtr.IsValid())
+	{
+		if (UMaterial const* BaseMaterial = BaseMaterialSoftPtr.LoadSynchronous())
+		{
+			MaterialValidationLibrary::FMaterialPermutationData PermutationData;
+			MaterialValidationLibrary::BuildPermutationData(*BaseMaterial, PermutationData);
+
+			// Only use resolved names if the layout matches the new snapshot.
+			if (PermutationData.StaticPropertyLayoutHash == NewDesc->AssetData.StaticPropertyLayoutHash)
+			{
+				for (FMaterialParameterInfo const& Info : PermutationData.StaticSwitchParameterInfo)
+				{
+					SwitchNames.Add(Info.Name);
+				}
+				for (FMaterialParameterInfo const& Info : PermutationData.ComponentMaskParameterInfo)
+				{
+					ComponentMaskNames.Add(Info.Name);
+				}
+			
+				bFoundParameterNames = true;
+			}
+		}
+	}
+
+	// Use default parameter names as a fallback.
+	if (!bFoundParameterNames)
+	{
+		{
+			static FName BaseName("StaticSwitch");
+			const int32 MaxSwitchNum = FMath::Max(OldDesc->AssetData.StaticSwitchNum, NewDesc->AssetData.StaticSwitchNum);
+			for (int32 Index = 0; Index < MaxSwitchNum; ++Index)
+			{
+				SwitchNames.Add(FName(BaseName, Index));
+			}
+		}
+		{
+			static const FName BaseName("ComponentMask");
+			const int32 MaxComponentMaskNum = FMath::Max(OldDesc->AssetData.ComponentMaskNum, NewDesc->AssetData.ComponentMaskNum);
+			for (int32 Index = 0; Index < MaxComponentMaskNum; ++Index)
+			{
+				ComponentMaskNames.Add(FName(BaseName, Index));
+			}
+		}
+	}
+
+	// Base material entry.
+	FMaterialInstanceDiffResult BaseDiff;
+	BaseDiff.InstancePath = InBaseMaterialPath;
+	MaterialValidationLibrary::BuildBaseMaterialPropertyDiff(OldDesc->AssetData, NewDesc->AssetData, SwitchNames, ComponentMaskNames, BaseDiff.PropertyDiffs);
+	if (BaseDiff.PropertyDiffs.Num())
+	{
+		BaseDiff.DiffType = EMaterialInstanceDiffType::Modified;
+		OutDiffs.Add(MoveTemp(BaseDiff));
+	}
+
+	// Instance entries.
+	const bool bLayoutsMatch = OldDesc->AssetData.StaticPropertyLayoutHash == NewDesc->AssetData.StaticPropertyLayoutHash;
+
+	TSet<FString> AllInstancePaths;
+	for (TPair<FString, int32> const& It : OldDesc->MaterialInstances) 
+	{
+		AllInstancePaths.Add(It.Key); 
+	}
+	for (TPair<FString, int32> const& It : NewDesc->MaterialInstances) 
+	{
+		AllInstancePaths.Add(It.Key); 
+	}
+
+	for (FString const& PathStr : AllInstancePaths)
+	{
+		FSoftObjectPath InstancePath = MaterialValidationLibrary::ConvertStringPathToSoftObjectPath(PathStr);
+
+		int32 const* OldIndexPtr = OldDesc->MaterialInstances.Find(PathStr);
+		if (OldIndexPtr == nullptr)
+		{
+			OutDiffs.Add({ .InstancePath = InstancePath, .DiffType = EMaterialInstanceDiffType::Added });
+			continue;
+		}
+		
+		int32 const* NewIndexPtr = NewDesc->MaterialInstances.Find(PathStr);
+		if (NewIndexPtr == nullptr)
+		{
+			OutDiffs.Add({ .InstancePath = InstancePath, .DiffType = EMaterialInstanceDiffType::Removed });
+			continue;
+		}
+
+		auto GetInstanceData = [](FMaterialValidationDesc const& Desc, int32 Index, int32 SwitchArrayNum, int32 MaskArrayNum) -> FMaterialValidationAssetData_MaterialInstance
+		{
+			if (Index == INDEX_NONE)
+			{
+				FMaterialValidationAssetData_MaterialInstance Default;
+				Default.StaticSwitchOverrideMask.SetNumZeroed(SwitchArrayNum);
+				Default.StaticSwitchOverrideValues.SetNumZeroed(SwitchArrayNum);
+				Default.ComponentMaskOverrideMask.SetNumZeroed(MaskArrayNum);
+				Default.ComponentMaskOverrideValues.SetNumZeroed(MaskArrayNum);
+				return Default;
+			}
+			if (Desc.MaterialInstanceAssetDatas.IsValidIndex(Index))
+			{
+				return Desc.MaterialInstanceAssetDatas[Index];
+			}
+			return FMaterialValidationAssetData_MaterialInstance{};
+		};
+
+		const FMaterialValidationAssetData_MaterialInstance OldData = GetInstanceData(*OldDesc, *OldIndexPtr, OldDesc->AssetData.StaticSwitchValues.Num(), OldDesc->AssetData.ComponentMaskValues.Num());
+		const FMaterialValidationAssetData_MaterialInstance NewData = GetInstanceData(*NewDesc, *NewIndexPtr, NewDesc->AssetData.StaticSwitchValues.Num(), NewDesc->AssetData.ComponentMaskValues.Num());
+
+		FMaterialInstanceDiffResult InstanceDiff;
+		InstanceDiff.InstancePath = InstancePath;
+		MaterialValidationLibrary::BuildInstancePropertyDiff(OldData, NewData, SwitchNames, ComponentMaskNames, bLayoutsMatch, InstanceDiff.PropertyDiffs);
+
+		if (InstanceDiff.PropertyDiffs.Num())
+		{
+			InstanceDiff.DiffType = EMaterialInstanceDiffType::Modified;
+			OutDiffs.Add(MoveTemp(InstanceDiff));
+		}
+	}
+
+	// Sort so that modified items appear before added/removed ones.
+	// Note that if the base material is in the array it will in the first entry. And because it is guaranteed to be EMaterialInstanceDiffType::Modified the StableSort will keep it there.
+	Algo::StableSort(OutDiffs, [](FMaterialInstanceDiffResult const& A, FMaterialInstanceDiffResult const& B)
+	{
+		return static_cast<uint8>(A.DiffType) < static_cast<uint8>(B.DiffType);
+	});
+}
+
 #undef LOCTEXT_NAMESPACE
diff --git a/Engine/Plugins/Editor/MaterialValidation/Source/Private/MaterialValidationLibrary.h b/Engine/Plugins/Editor/MaterialValidation/Source/Private/MaterialValidationLibrary.h
--- a/Engine/Plugins/Editor/MaterialValidation/Source/Private/MaterialValidationLibrary.h
+++ b/Engine/Plugins/Editor/MaterialValidation/Source/Private/MaterialValidationLibrary.h
@@ -6,76 +6,15 @@
 #include "Kismet/BlueprintFunctionLibrary.h"
 #include "MaterialValidationLibrary.generated.h"
 
-struct FMaterialValidationAssetData_MaterialInstance;
 class UMaterial;
 class UMaterialInterface;
 class UMaterialInstanceConstant;
+class UMaterialValidationGroup;
 
-/** Enumeration of material instance property types. */
-UENUM(BlueprintType)
-enum class EMaterialInstanceColumnType : uint8
-{
-	Default,
-	StaticProperty,
-	UsageFlag,
-	StaticSwitch,
-	ComponentMask,
-	Count
-};
-
-/** 
- * Structure that collects all the resolved asset data for one material hierarchy in the material asset database. 
- * Typically this is extracted from the compact form stored in a UMaterialValidationGroup and used in UI views or for further analysis.
- */
-USTRUCT(BlueprintType)
-struct FMaterialDatabaseAssetHierarchyInfo
-{
-	GENERATED_BODY()
-
-	/** Array of material asset paths. This includes the base UMaterial asset in the first entry, and the child UMaterialInstance assets after that. */
-	UPROPERTY(BlueprintReadOnly, Category = Material)
-	TArray<FSoftObjectPath> MaterialPaths;
-	/** Array of asset data info. This is synced with the MaterialPaths array. */
-	UPROPERTY()
-	TArray<FMaterialValidationAssetData_MaterialInstance> MaterialAssetDatas;
-	/** Array of all StaticSwitch names. This matches the layout of all StaticSwitchOverrideValues arrays in the MaterialAssetDatas. */
-	UPROPERTY()
-	TArray<FName> StaticSwitchNames;
-	/** Array of all ComponentMask names. This matches the layout of all ComponentMaskOverrideValues arrays in the MaterialAssetDatas. */
-	UPROPERTY()
-	TArray<FName> ComponentMaskNames;
-};
-
-/** Description of a material property in the material asset database. */
-USTRUCT(BlueprintType)
-struct FMaterialDatabaseAssetPropertyDesc
-{
-	GENERATED_BODY()
-
-	/** Id for the property. */
-	UPROPERTY(BlueprintReadOnly, Category = Material)
-	FName Id;
-	/** Readable name for the property. */
-	UPROPERTY(BlueprintReadOnly, Category = Material)
-	FName Name;
-	/** Category for the property. */
-	UPROPERTY(BlueprintReadOnly, Category = Material)
-	EMaterialInstanceColumnType Category = EMaterialInstanceColumnType::Default;
-};
-
-/** Description of a resolved material value in the material asset database. */
-USTRUCT(BlueprintType)
-struct FMaterialDatabaseAssetPropertyValue
-{
-	GENERATED_BODY()
-
-	/** String value for the property. */
-	UPROPERTY(BlueprintReadOnly, Category = Material)
-	FString Value;
-	/** true if the value for the property is different from the value on its parent material. */
-	UPROPERTY(BlueprintReadOnly, Category = Material)
-	bool bIsModifiedFromParent = false;
-};
+struct FMaterialDatabaseAssetHierarchyInfo;
+struct FMaterialDatabaseAssetPropertyDesc;
+struct FMaterialDatabaseAssetPropertyValue;
+struct FMaterialInstanceDiffResult;
 
 /** Material validation functionality exposed as a library of Blueprint functions. */
 UCLASS(MinimalAPI, meta = (ScriptName = "MaterialValidationLibrary"))
@@ -208,4 +147,18 @@ class UMaterialValidationLibrary : public UBlueprintFunctionLibrary
 	 */
 	UFUNCTION(BlueprintCallable, Category = Material)
 	static void GetMaterialPropertyValues(FMaterialDatabaseAssetHierarchyInfo const& InInfo, int32 InMaterialIndex, TArray<FMaterialDatabaseAssetPropertyValue>& OutValues);
+
+	/**
+	 * Get the diff between two versions of a UMaterialValidationGroup for a single base material.
+	 * Compares the raw stored data. There is no asset registry dependency and it works for deleted instances.
+	 * The first entry in OutDiffs always contains the base material entry if it changed.
+	 *
+	 * @param	InGroupOld				The older version of the group for comparison.
+	 * @param	InGroupNew				The newer version of the group for comparison.
+	 * @param	InBaseMaterialPath		The path of the base material to find in the groups and diff.
+	 * @param	bAllowAssetLoad			If true then the base UMaterial is loaded to resolve static parameter names. If false then static parameters are given generic names.
+	 * @param	OutDiffs				Returns one entry per changed material instance. Empty if no differences found.
+	 */
+	UFUNCTION(BlueprintCallable, Category = Material)
+	static void GetMaterialValidationDescDiff(UMaterialValidationGroup const* InGroupOld, UMaterialValidationGroup const* InGroupNew, FSoftObjectPath const& InBaseMaterialPath, bool bAllowAssetLoad, TArray<FMaterialInstanceDiffResult>& OutDiffs);
 };
diff --git a/Engine/Plugins/Editor/MaterialValidation/Source/Private/MaterialValidationLibraryTypes.h b/Engine/Plugins/Editor/MaterialValidation/Source/Private/MaterialValidationLibraryTypes.h
new file mode 100644
--- /dev/null
+++ b/Engine/Plugins/Editor/MaterialValidation/Source/Private/MaterialValidationLibraryTypes.h
@@ -0,0 +1,119 @@
+// Copyright Epic Games, Inc. All Rights Reserved.
+
+#pragma once
+
+#include "MaterialValidationLibraryTypes.generated.h"
+
+struct FMaterialValidationAssetData_MaterialInstance;
+
+/** Enumeration of material instance property types. */
+UENUM(BlueprintType)
+enum class EMaterialInstanceColumnType : uint8
+{
+	Default,
+	StaticProperty,
+	UsageFlag,
+	StaticSwitch,
+	ComponentMask,
+	Count
+};
+
+/** Whether a material instance was added, removed, or modified between two revisions. */
+UENUM(BlueprintType)
+enum class EMaterialInstanceDiffType : uint8
+{
+	Modified,
+	Added,
+	Removed,
+};
+
+/** 
+ * Structure that collects all the resolved asset data for one material hierarchy in the material asset database. 
+ * Typically this is extracted from the compact form stored in a UMaterialValidationGroup and used in UI views or for further analysis.
+ */
+USTRUCT(BlueprintType)
+struct FMaterialDatabaseAssetHierarchyInfo
+{
+	GENERATED_BODY()
+
+	/** Array of material asset paths. This includes the base UMaterial asset in the first entry, and the child UMaterialInstance assets after that. */
+	UPROPERTY(BlueprintReadOnly, Category = Material)
+	TArray<FSoftObjectPath> MaterialPaths;
+	/** Array of asset data info. This is synced with the MaterialPaths array. */
+	UPROPERTY()
+	TArray<FMaterialValidationAssetData_MaterialInstance> MaterialAssetDatas;
+	/** Array of all StaticSwitch names. This matches the layout of all StaticSwitchOverrideValues arrays in the MaterialAssetDatas. */
+	UPROPERTY()
+	TArray<FName> StaticSwitchNames;
+	/** Array of all ComponentMask names. This matches the layout of all ComponentMaskOverrideValues arrays in the MaterialAssetDatas. */
+	UPROPERTY()
+	TArray<FName> ComponentMaskNames;
+};
+
+/** Description of a material property in the material asset database. */
+USTRUCT(BlueprintType)
+struct FMaterialDatabaseAssetPropertyDesc
+{
+	GENERATED_BODY()
+
+	/** Id for the property. */
+	UPROPERTY(BlueprintReadOnly, Category = Material)
+	FName Id;
+	/** Readable name for the property. */
+	UPROPERTY(BlueprintReadOnly, Category = Material)
+	FName Name;
+	/** Category for the property. */
+	UPROPERTY(BlueprintReadOnly, Category = Material)
+	EMaterialInstanceColumnType Category = EMaterialInstanceColumnType::Default;
+};
+
+/** Description of a resolved material value in the material asset database. */
+USTRUCT(BlueprintType)
+struct FMaterialDatabaseAssetPropertyValue
+{
+	GENERATED_BODY()
+
+	/** String value for the property. */
+	UPROPERTY(BlueprintReadOnly, Category = Material)
+	FString Value;
+	/** true if the value for the property is different from the value on its parent material. */
+	UPROPERTY(BlueprintReadOnly, Category = Material)
+	bool bIsModifiedFromParent = false;
+};
+
+/** A single property difference between two versions of a material instance. */
+USTRUCT(BlueprintType)
+struct FMaterialInstancePropertyDiff
+{
+	GENERATED_BODY()
+
+	/** Human-readable name of the property. */
+	UPROPERTY(BlueprintReadOnly, Category = Material)
+	FName PropertyName;
+	/** Category for the property. */
+	UPROPERTY(BlueprintReadOnly, Category = Material)
+	EMaterialInstanceColumnType Category = EMaterialInstanceColumnType::Default;
+	/** String representation of the value in the old revision. "[Inherited]" if not overridden. */
+	UPROPERTY(BlueprintReadOnly, Category = Material)
+	FString OldValue;
+	/** String representation of the value in the new revision. "[Inherited]" if not overridden. */
+	UPROPERTY(BlueprintReadOnly, Category = Material)
+	FString NewValue;
+};
+
+/** Diff result for a single material instance between two versions of a UMaterialValidationGroup. */
+USTRUCT(BlueprintType)
+struct FMaterialInstanceDiffResult
+{
+	GENERATED_BODY()
+
+	/** Asset path of the material or material instance. */
+	UPROPERTY(BlueprintReadOnly, Category = Material)
+	FSoftObjectPath InstancePath;
+	/** Whether this instance was added, removed, or modified. */
+	UPROPERTY(BlueprintReadOnly, Category = Material)
+	EMaterialInstanceDiffType DiffType = EMaterialInstanceDiffType::Modified;
+	/** Per-property differences. This will be empty for Added and Removed entries. */
+	UPROPERTY(BlueprintReadOnly, Category = Material)
+	TArray<FMaterialInstancePropertyDiff> PropertyDiffs;
+};
diff --git a/Engine/Plugins/Editor/MaterialValidation/Source/Private/MaterialValidationToolkit.cpp b/Engine/Plugins/Editor/MaterialValidation/Source/Private/MaterialValidationToolkit.cpp
--- a/Engine/Plugins/Editor/MaterialValidation/Source/Private/MaterialValidationToolkit.cpp
+++ b/Engine/Plugins/Editor/MaterialValidation/Source/Private/MaterialValidationToolkit.cpp
@@ -7,7 +7,9 @@
 #include "Framework/Application/SlateApplication.h"
 #include "Materials/MaterialInterface.h"
 #include "MaterialValidationGroup.h"
+#include "MaterialValidationGroupDiff.h"
 #include "MaterialValidationLibrary.h"
+#include "MaterialValidationLibraryTypes.h"
 #include "Misc/NotifyHook.h"
 #include "Modules/ModuleManager.h"
 #include "PropertyCustomizationHelpers.h"  
@@ -181,8 +183,11 @@ static TSharedPtr<SExperimentsTable> CreateExperimentsTable()
 }
 
 /** Helper function to create an SWidget for AssetPath with text and a browse button. */
-static TSharedRef<SWidget> GenerateAssetPathWidget(FSoftObjectPath InAssetPath)
+TSharedRef<SWidget> GenerateAssetPathWidget(FSoftObjectPath InAssetPath, FSlateColor InColor)
 {
+	const FAssetData AssetData = IAssetRegistry::Get()->GetAssetByObjectPath(InAssetPath);
+	const EVisibility BrowseVisibility = AssetData.IsValid() ? EVisibility::Visible : EVisibility::Hidden;
+
 	return SNew(SBox)
 		.HAlign(HAlign_Fill)
 		.VAlign(VAlign_Center)
@@ -195,28 +200,31 @@ static TSharedRef<SWidget> GenerateAssetPathWidget(FSoftObjectPath InAssetPath)
 					SNew(STextBlock)
 						.Text(FText::FromString(InAssetPath.ToString()))
 						.ToolTipText(FText::FromString(InAssetPath.ToString()))
+						.ColorAndOpacity(InColor)
 				]
 				+ SHorizontalBox::Slot()
 				.AutoWidth()
 				.VAlign(VAlign_Center)
 				.Padding(4.0f, 0.0f)
 				[
-					PropertyCustomizationHelpers::MakeBrowseButton(
-						FSimpleDelegate::CreateLambda([InAssetPath]()
-						{
-							const FAssetData AssetData = IAssetRegistry::Get()->GetAssetByObjectPath(InAssetPath);
-							
-							const FModifierKeysState KeyState = FSlateApplication::Get().GetModifierKeys();
-							if (KeyState.IsAltDown())
-							{
-								GEditor->EditObject(AssetData.GetAsset());
-							}
-							else
+					SNew(SBox)
+					.Visibility(BrowseVisibility)
+					[
+						PropertyCustomizationHelpers::MakeBrowseButton(
+							FSimpleDelegate::CreateLambda([AssetData]()
 							{
-								GEditor->SyncBrowserToObjects({ AssetData });
-							}
-						}),
-						LOCTEXT("BrowseTooltip", "Browse to Asset in Content Browser"))
+								const FModifierKeysState KeyState = FSlateApplication::Get().GetModifierKeys();
+								if (KeyState.IsAltDown())
+								{
+									GEditor->EditObject(AssetData.GetAsset());
+								}
+								else
+								{
+									GEditor->SyncBrowserToObjects({ AssetData });
+								}
+							}),
+							LOCTEXT("BrowseTooltip", "Browse to Asset in Content Browser"))
+					]
 				]
 		];
 }
@@ -245,7 +253,7 @@ public:
 		{
 			if (ColumnName == FMaterialStats::ColumnId_AssetPath)
 			{
-				return GenerateAssetPathWidget(RowItem->AssetPath);
+				return GenerateAssetPathWidget(RowItem->AssetPath, FSlateColor::UseForeground());
 			}
 			if (ColumnName == FMaterialStats::ColumnId_MaterialInstanceCount)
 			{
@@ -650,7 +658,7 @@ static TSharedRef<SWidget> GenerateWidgetForStaticProperty(TSharedPtr<FMaterialI
 
 	if constexpr(std::is_same<std::remove_cv_t<std::remove_reference_t<ReturnType>>, FSoftObjectPath>::value)
 	{
-		return GenerateAssetPathWidget((*GetterPtr)(*RowItem));
+		return GenerateAssetPathWidget((*GetterPtr)(*RowItem), FSlateColor::UseForeground());
 	}
 	else if constexpr(TIsFloatingPoint<ReturnType>::Value)
 	{
@@ -1810,4 +1818,18 @@ EAssetCommandResult UAssetDefinition_MaterialValidationGroup::OpenAssets(const F
 	return EAssetCommandResult::Handled;
 }
 
+EAssetCommandResult UAssetDefinition_MaterialValidationGroup::PerformAssetDiff(const FAssetDiffArgs& DiffArgs) const
+{
+	UMaterialValidationGroup const* OldGroup = Cast<UMaterialValidationGroup>(DiffArgs.OldAsset);
+	UMaterialValidationGroup const* NewGroup = Cast<UMaterialValidationGroup>(DiffArgs.NewAsset);
+
+	if (OldGroup == nullptr && NewGroup == nullptr)
+	{
+		return EAssetCommandResult::Unhandled;
+	}
+
+	SMaterialValidationGroupDiff::CreateDiffWindow(OldGroup, NewGroup, DiffArgs.OldRevision, DiffArgs.NewRevision);
+	return EAssetCommandResult::Handled;
+}
+
 #undef LOCTEXT_NAMESPACE
diff --git a/Engine/Plugins/Editor/MaterialValidation/Source/Private/MaterialValidationToolkit.h b/Engine/Plugins/Editor/MaterialValidation/Source/Private/MaterialValidationToolkit.h
--- a/Engine/Plugins/Editor/MaterialValidation/Source/Private/MaterialValidationToolkit.h
+++ b/Engine/Plugins/Editor/MaterialValidation/Source/Private/MaterialValidationToolkit.h
@@ -16,4 +16,10 @@ class UAssetDefinition_MaterialValidationGroup : public UAssetDefinitionDefault
 	virtual TSoftClassPtr<UObject> GetAssetClass() const override;
 	virtual TConstArrayView<FAssetCategoryPath> GetAssetCategories() const override;
 	virtual EAssetCommandResult OpenAssets(const FAssetOpenArgs& OpenArgs) const override;
+	virtual EAssetCommandResult PerformAssetDiff(const FAssetDiffArgs& DiffArgs) const override;
 };
+
+namespace MaterialValidation
+{
+	TSharedRef<SWidget> GenerateAssetPathWidget(FSoftObjectPath InAssetPath, FSlateColor InColor);
+}
