MaterialValidation: Use fuzzy search in toolkit and diff tool search bars.
#rnx
#Synced-CL 51607007

diff --git a/Engine/Plugins/Editor/MaterialValidation/Source/Private/MaterialValidationGroupDiff.cpp b/Engine/Plugins/Editor/MaterialValidation/Source/Private/MaterialValidationGroupDiff.cpp
--- a/Engine/Plugins/Editor/MaterialValidation/Source/Private/MaterialValidationGroupDiff.cpp
+++ b/Engine/Plugins/Editor/MaterialValidation/Source/Private/MaterialValidationGroupDiff.cpp
@@ -9,6 +9,7 @@
 #include "MaterialValidationLibrary.h"
 #include "MaterialValidationLibraryTypes.h"
 #include "MaterialValidationToolkit.h"
+#include "Misc/TextFilterExpressionEvaluator.h"
 #include "PropertyCustomizationHelpers.h"
 #include "Widgets/Input/SButton.h"
 #include "Widgets/Input/SSearchBox.h"
@@ -363,9 +364,13 @@ private:
 	void UpdateFilter()
 	{
 		FilteredRows.Reset();
+
+		FTextFilterExpressionEvaluator Filter(ETextFilterExpressionEvaluatorMode::BasicString);
+		Filter.SetFilterText(FText::FromString(SearchText));
+
 		for (TSharedPtr<FMaterialDiffStats> const& Row : AllRows)
 		{
-			if (SearchText.IsEmpty() || Row->AssetPath.ToString().Contains(SearchText, ESearchCase::IgnoreCase))
+			if (SearchText.IsEmpty() || Filter.TestTextFilter(FBasicStringFilterExpressionContext(Row->AssetPath.ToString())))
 			{
 				FilteredRows.Add(Row);
 			}
@@ -725,11 +730,17 @@ private:
 	{
 		FilteredRootNodes.Reset();
 
+		FTextFilterExpressionEvaluator AssetFilter(ETextFilterExpressionEvaluatorMode::BasicString);
+		AssetFilter.SetFilterText(FText::FromString(InstanceSearchText));
+
+		FTextFilterExpressionEvaluator PropertyFilter(ETextFilterExpressionEvaluatorMode::BasicString);
+		PropertyFilter.SetFilterText(FText::FromString(PropertySearchText));
+
 		for (TSharedPtr<FMaterialDiffTreeNode> const& Node : AllRootNodes)
 		{
 			// Instance path filter.
-			if (!InstanceSearchText.IsEmpty() &&
-				!Node->DiffResult.InstancePath.ToString().Contains(InstanceSearchText, ESearchCase::IgnoreCase))
+			if (!InstanceSearchText.IsEmpty() && 
+				!AssetFilter.TestTextFilter(FBasicStringFilterExpressionContext(Node->DiffResult.InstancePath.ToString())))
 			{
 				continue;
 			}
@@ -747,7 +758,7 @@ private:
 				bool bHasAnyMatch = false;
 				for (TSharedPtr<FMaterialDiffTreeNode> const& Child : Node->Children)
 				{
-					if (Child->PropertyDiff.PropertyName.ToString().Contains(PropertySearchText, ESearchCase::IgnoreCase))
+					if (PropertyFilter.TestTextFilter(FBasicStringFilterExpressionContext(Child->PropertyDiff.PropertyName.ToString())))
 					{
 						bHasAnyMatch = true;
 						break;
diff --git a/Engine/Plugins/Editor/MaterialValidation/Source/Private/MaterialValidationToolkit.cpp b/Engine/Plugins/Editor/MaterialValidation/Source/Private/MaterialValidationToolkit.cpp
--- a/Engine/Plugins/Editor/MaterialValidation/Source/Private/MaterialValidationToolkit.cpp
+++ b/Engine/Plugins/Editor/MaterialValidation/Source/Private/MaterialValidationToolkit.cpp
@@ -11,6 +11,7 @@
 #include "MaterialValidationLibrary.h"
 #include "MaterialValidationLibraryTypes.h"
 #include "Misc/NotifyHook.h"
+#include "Misc/TextFilterExpressionEvaluator.h"
 #include "Modules/ModuleManager.h"
 #include "PropertyCustomizationHelpers.h"  
 #include "PropertyEditorModule.h"
@@ -421,9 +422,12 @@ public:
 	{
 		FilteredRows.Reset();
 
+		FTextFilterExpressionEvaluator Filter(ETextFilterExpressionEvaluatorMode::BasicString);
+		Filter.SetFilterText(FText::FromString(SearchText));
+
 		for (const TSharedPtr<FMaterialStats>& Row : AllRows)
 		{
-			if (SearchText.IsEmpty() || Row->AssetPath.ToString().Contains(SearchText, ESearchCase::IgnoreCase))
+			if (SearchText.IsEmpty() || Filter.TestTextFilter(FBasicStringFilterExpressionContext(Row->AssetPath.ToString())))
 			{
 				FilteredRows.Add(Row);
 			}
@@ -1235,9 +1239,12 @@ protected:
 
 		TSet<uint32> PermutationHashesSeen;
 
+		FTextFilterExpressionEvaluator Filter(ETextFilterExpressionEvaluatorMode::BasicString);
+		Filter.SetFilterText(FText::FromString(RowSearchText));
+
 		for (const TSharedPtr<FMaterialInstanceRowData>& Row : AllRows)
 		{
-			if (RowSearchText.IsEmpty() || Row->AssetPath.ToString().Contains(RowSearchText, ESearchCase::IgnoreCase))
+			if (RowSearchText.IsEmpty() || Filter.TestTextFilter(FBasicStringFilterExpressionContext(Row->AssetPath.ToString())))
 			{
 				if (bUniqueInstances)
 				{
