MaterialValidation: Rename ColumnType -> Category for clarity now that the enum is used for more than just column filtering.
#rnx
#Synced-CL 51620377

diff --git a/Engine/Plugins/Editor/MaterialValidation/Config/DefaultMaterialValidation.ini b/Engine/Plugins/Editor/MaterialValidation/Config/DefaultMaterialValidation.ini
new file mode 100644
--- /dev/null
+++ b/Engine/Plugins/Editor/MaterialValidation/Config/DefaultMaterialValidation.ini
@@ -0,0 +1,3 @@
+[CoreRedirects]
+
++EnumRedirects=(OldName="/Script/MaterialValidation.EMaterialInstanceColumnType",NewName="/Script/MaterialValidation.EMaterialPropertyCategory")
diff --git a/Engine/Plugins/Editor/MaterialValidation/Source/Private/MaterialValidationLibrary.cpp b/Engine/Plugins/Editor/MaterialValidation/Source/Private/MaterialValidationLibrary.cpp
--- a/Engine/Plugins/Editor/MaterialValidation/Source/Private/MaterialValidationLibrary.cpp
+++ b/Engine/Plugins/Editor/MaterialValidation/Source/Private/MaterialValidationLibrary.cpp
@@ -1583,7 +1583,7 @@ static FString FormatComponentMask(uint32 InValue)
 static const TCHAR* GInheritedValueString = TEXT("[Inherited]");
 
 /** Add an entry to InOutDiffs if old and new string values differ. */
-static void AddDiffIfChanged(FName InName, EMaterialInstanceColumnType InCategory, FString const& InOldValue, FString const& InNewValue, TArray<FMaterialInstancePropertyDiff>& InOutDiffs)
+static void AddDiffIfChanged(FName InName, EMaterialPropertyCategory InCategory, FString const& InOldValue, FString const& InNewValue, TArray<FMaterialInstancePropertyDiff>& InOutDiffs)
 {
 	if (InOldValue != InNewValue)
 	{
@@ -1601,7 +1601,7 @@ static void BuildStaticPropertyDiff(
 	{
 		AddDiffIfChanged(
 			StaticPropertyInfo.PropertyName,
-			EMaterialInstanceColumnType::StaticProperty,
+			EMaterialPropertyCategory::StaticProperty,
 			StaticPropertyInfo.GetPropertyValueAsStringFunction(InOld),
 			StaticPropertyInfo.GetPropertyValueAsStringFunction(InNew),
 			InOutDiffs);
@@ -1622,7 +1622,7 @@ static void BuildStaticPropertyDiff(
 		const bool bNewInherited = !StaticPropertyInfo.GetIsPropertyModifiedFunction(InNewOverrideFlags);
 		AddDiffIfChanged(
 			StaticPropertyInfo.PropertyName,
-			EMaterialInstanceColumnType::StaticProperty,
+			EMaterialPropertyCategory::StaticProperty,
 			bOldInherited ? GInheritedValueString : StaticPropertyInfo.GetPropertyValueAsStringFunction(InOld),
 			bNewInherited ? GInheritedValueString : StaticPropertyInfo.GetPropertyValueAsStringFunction(InNew),
 			InOutDiffs);
@@ -1640,7 +1640,7 @@ static void BuildUsageFlagDiff(
 		const FName UsageName(*UMaterialInterface::GetUsageName((EMaterialUsage)Index));
 		const FString OldValue = (InOld.UsageFlags & (1u << Index)) != 0 ? TEXT("1") : TEXT("0");
 		const FString NewValue = (InNew.UsageFlags & (1u << Index)) != 0 ? TEXT("1") : TEXT("0");
-		AddDiffIfChanged(UsageName, EMaterialInstanceColumnType::UsageFlag, OldValue, NewValue, InOutDiffs);
+		AddDiffIfChanged(UsageName, EMaterialPropertyCategory::UsageFlag, OldValue, NewValue, InOutDiffs);
 	}
 }
 
@@ -1657,7 +1657,7 @@ static void BuildUsageFlagDiff(
 		const FName UsageName(*UMaterialInterface::GetUsageName((EMaterialUsage)Index));
 		const FString OldValue = !(InOldOverrideFlags.Override_UsageFlags & (1u << Index)) ? GInheritedValueString : (InOld.UsageFlags & (1u << Index)) != 0 ? TEXT("1") : TEXT("0");
 		const FString NewValue = !(InNewOverrideFlags.Override_UsageFlags & (1u << Index)) ? GInheritedValueString : (InNew.UsageFlags & (1u << Index)) != 0 ? TEXT("1") : TEXT("0");
-		AddDiffIfChanged(UsageName, EMaterialInstanceColumnType::UsageFlag, OldValue, NewValue, InOutDiffs);
+		AddDiffIfChanged(UsageName, EMaterialPropertyCategory::UsageFlag, OldValue, NewValue, InOutDiffs);
 	}
 }
 
@@ -1672,7 +1672,7 @@ static void BuildStaticSwitchDiff(
 	{
 		const FString OldValue = GetPackedBit(InOld, Index) ? TEXT("1") : TEXT("0");
 		const FString NewValue = GetPackedBit(InNew, Index) ? TEXT("1") : TEXT("0");
-		AddDiffIfChanged(InStaticSwitchNames[Index], EMaterialInstanceColumnType::StaticSwitch, OldValue, NewValue, InOutDiffs);
+		AddDiffIfChanged(InStaticSwitchNames[Index], EMaterialPropertyCategory::StaticSwitch, OldValue, NewValue, InOutDiffs);
 	}
 }
 
@@ -1689,7 +1689,7 @@ static void BuildStaticSwitchDiff(
 	{
 		const FString OldValue = !GetPackedBit(InOldOverrideMask, Index) ? GInheritedValueString : GetPackedBit(InOld, Index) ? TEXT("1") : TEXT("0");
 		const FString NewValue = !GetPackedBit(InNewOverrideMask, Index) ? GInheritedValueString : GetPackedBit(InNew, Index) ? TEXT("1") : TEXT("0");
-		AddDiffIfChanged(InStaticSwitchNames[Index], EMaterialInstanceColumnType::StaticSwitch, OldValue, NewValue, InOutDiffs);
+		AddDiffIfChanged(InStaticSwitchNames[Index], EMaterialPropertyCategory::StaticSwitch, OldValue, NewValue, InOutDiffs);
 	}
 }
 
@@ -1704,7 +1704,7 @@ static void BuildComponentMaskDiff(
 	{
 		const FString OldValue = FormatComponentMask(GetPackedNibble(InOld, Index));
 		const FString NewValue = FormatComponentMask(GetPackedNibble(InNew, Index));
-		AddDiffIfChanged(InComponentMaskNames[Index], EMaterialInstanceColumnType::ComponentMask, OldValue, NewValue, InOutDiffs);
+		AddDiffIfChanged(InComponentMaskNames[Index], EMaterialPropertyCategory::ComponentMask, OldValue, NewValue, InOutDiffs);
 	}
 }
 
@@ -1721,7 +1721,7 @@ static void BuildComponentMaskDiff(
 	{
 		const FString OldValue = !GetPackedNibble(InOldOverrideMask, Index) ? GInheritedValueString : FormatComponentMask(GetPackedNibble(InOld, Index));
 		const FString NewValue = !GetPackedNibble(InNewOverrideMask, Index) ? GInheritedValueString : FormatComponentMask(GetPackedNibble(InNew, Index));
-		AddDiffIfChanged(InComponentMaskNames[Index], EMaterialInstanceColumnType::ComponentMask, OldValue, NewValue, InOutDiffs);
+		AddDiffIfChanged(InComponentMaskNames[Index], EMaterialPropertyCategory::ComponentMask, OldValue, NewValue, InOutDiffs);
 	}
 }
 
@@ -1801,19 +1801,19 @@ void UMaterialValidationLibrary::GetMaterialProperties(FMaterialDatabaseAssetHie
 
 	for (MaterialValidationLibrary::FMaterialInstanceStaticPropertyInfo const& StaticPropertyInfo : MaterialValidationLibrary::StaticPropertyInfos)
 	{
-		OutProperties.Add({ .Id = StaticPropertyInfo.PropertyId, .Name = StaticPropertyInfo.PropertyName, .Category = EMaterialInstanceColumnType::StaticProperty });
+		OutProperties.Add({ .Id = StaticPropertyInfo.PropertyId, .Name = StaticPropertyInfo.PropertyName, .Category = EMaterialPropertyCategory::StaticProperty });
 	}
 	for (int32 Index = 0; Index < MATUSAGE_MAX; ++Index)
 	{
-		OutProperties.Add({ .Id = MaterialValidationLibrary::UsageFlagPropertyIds[Index], .Name = FName(*UMaterialInterface::GetUsageName((EMaterialUsage)Index)), .Category = EMaterialInstanceColumnType::UsageFlag});
+		OutProperties.Add({ .Id = MaterialValidationLibrary::UsageFlagPropertyIds[Index], .Name = FName(*UMaterialInterface::GetUsageName((EMaterialUsage)Index)), .Category = EMaterialPropertyCategory::UsageFlag});
 	}
 	for (const FName Name : InInfo.StaticSwitchNames)
 	{
-		OutProperties.Add({ .Id = Name, .Name = Name, .Category = EMaterialInstanceColumnType::StaticSwitch});
+		OutProperties.Add({ .Id = Name, .Name = Name, .Category = EMaterialPropertyCategory::StaticSwitch});
 	}
 	for (const FName Name : InInfo.ComponentMaskNames)
 	{
-		OutProperties.Add({ .Id = Name, .Name = Name, .Category = EMaterialInstanceColumnType::ComponentMask });
+		OutProperties.Add({ .Id = Name, .Name = Name, .Category = EMaterialPropertyCategory::ComponentMask });
 	}
 }
 	
diff --git a/Engine/Plugins/Editor/MaterialValidation/Source/Private/MaterialValidationLibraryTypes.h b/Engine/Plugins/Editor/MaterialValidation/Source/Private/MaterialValidationLibraryTypes.h
--- a/Engine/Plugins/Editor/MaterialValidation/Source/Private/MaterialValidationLibraryTypes.h
+++ b/Engine/Plugins/Editor/MaterialValidation/Source/Private/MaterialValidationLibraryTypes.h
@@ -8,7 +8,7 @@ struct FMaterialValidationAssetData_MaterialInstance;
 
 /** Enumeration of material instance property types. */
 UENUM(BlueprintType)
-enum class EMaterialInstanceColumnType : uint8
+enum class EMaterialPropertyCategory : uint8
 {
 	Default,
 	StaticProperty,
@@ -64,7 +64,7 @@ struct FMaterialDatabaseAssetPropertyDesc
 	FName Name;
 	/** Category for the property. */
 	UPROPERTY(BlueprintReadOnly, Category = Material)
-	EMaterialInstanceColumnType Category = EMaterialInstanceColumnType::Default;
+	EMaterialPropertyCategory Category = EMaterialPropertyCategory::Default;
 };
 
 /** Description of a resolved material value in the material asset database. */
@@ -92,7 +92,7 @@ struct FMaterialInstancePropertyDiff
 	FName PropertyName;
 	/** Category for the property. */
 	UPROPERTY(BlueprintReadOnly, Category = Material)
-	EMaterialInstanceColumnType Category = EMaterialInstanceColumnType::Default;
+	EMaterialPropertyCategory Category = EMaterialPropertyCategory::Default;
 	/** String representation of the value in the old revision. "[Inherited]" if not overridden. */
 	UPROPERTY(BlueprintReadOnly, Category = Material)
 	FString OldValue;
diff --git a/Engine/Plugins/Editor/MaterialValidation/Source/Private/MaterialValidationToolkit.cpp b/Engine/Plugins/Editor/MaterialValidation/Source/Private/MaterialValidationToolkit.cpp
--- a/Engine/Plugins/Editor/MaterialValidation/Source/Private/MaterialValidationToolkit.cpp
+++ b/Engine/Plugins/Editor/MaterialValidation/Source/Private/MaterialValidationToolkit.cpp
@@ -530,12 +530,12 @@ struct FMaterialInstanceColumnData
 	/** FText label used for display name. */
 	FText Label;
 	/** Column category used for filtering. */
-	EMaterialInstanceColumnType ColumnType;
+	EMaterialPropertyCategory Category;
 
-	FMaterialInstanceColumnData(FName const& InId, FText const& InLabel, EMaterialInstanceColumnType InColumnType)
+	FMaterialInstanceColumnData(FName const& InId, FText const& InLabel, EMaterialPropertyCategory InCategory)
 		: Id(InId)
 		, Label(InLabel)
-		, ColumnType(InColumnType)
+		, Category(InCategory)
 	{
 	}
 };
@@ -563,18 +563,18 @@ struct FMaterialInstanceRowData
 	/** Full resolved asset data description for the UMaterialInstance pulled from the material validation database. */
 	FMaterialValidationAssetData_MaterialInstance AssetData;
 
-	/** Per-column-type hash data, derived from instance data */
-	uint32 PerColumnTypeHash[(int32)EMaterialInstanceColumnType::Count] = {};
+	/** Per-category hash data, derived from instance data */
+	uint32 PerCategoryHash[(int32)EMaterialPropertyCategory::Count] = {};
 
 	FMaterialInstanceRowData(FSoftObjectPath const& InPath, int32 InHierarchyIndex, FMaterialValidationAssetData_MaterialInstance const& InAssetData)
 		: AssetPath(InPath)
 		, HierarchyIndex(InHierarchyIndex)
 		, AssetData(InAssetData)
 	{
-		PerColumnTypeHash[(int32)EMaterialInstanceColumnType::StaticProperty] = GetNoUsageHash(InAssetData.StaticProperties);
-		PerColumnTypeHash[(int32)EMaterialInstanceColumnType::UsageFlag] = InAssetData.StaticProperties.UsageFlags;
-		PerColumnTypeHash[(int32)EMaterialInstanceColumnType::StaticSwitch] = GetArrayHash(InAssetData.StaticSwitchOverrideValues);
-		PerColumnTypeHash[(int32)EMaterialInstanceColumnType::ComponentMask] = GetArrayHash(InAssetData.ComponentMaskOverrideValues);
+		PerCategoryHash[(int32)EMaterialPropertyCategory::StaticProperty] = GetNoUsageHash(InAssetData.StaticProperties);
+		PerCategoryHash[(int32)EMaterialPropertyCategory::UsageFlag] = InAssetData.StaticProperties.UsageFlags;
+		PerCategoryHash[(int32)EMaterialPropertyCategory::StaticSwitch] = GetArrayHash(InAssetData.StaticSwitchOverrideValues);
+		PerCategoryHash[(int32)EMaterialPropertyCategory::ComponentMask] = GetArrayHash(InAssetData.ComponentMaskOverrideValues);
 	}
 };
 
@@ -681,7 +681,7 @@ using PropertyGenerateWidgetFunction = TSharedRef<SWidget>(*)(TSharedPtr<FMateri
 /** Struct containing information needed to generate table entries for EMaterialInstanceStaticProperties. */
 struct FMaterialInstanceStaticPropertyInfo
 {
-	EMaterialInstanceColumnType Type;
+	EMaterialPropertyCategory Category;
 	FName ColumnId;
 	FString PropertyName;
 	PropertySortFunction SortFunction;
@@ -690,18 +690,18 @@ struct FMaterialInstanceStaticPropertyInfo
 
 /** Array for table generation data for EMaterialInstanceStaticProperties. */
 static const FMaterialInstanceStaticPropertyInfo StaticPropertyInfos[(int32)EMaterialInstanceStaticProperties::Count] = {
-	{EMaterialInstanceColumnType::Default, "CID_AssetPath", "Asset Path", &SortByStaticProperty<&FMaterialInstanceStaticPropertyGetters::GetAssetPathIndex>, &GenerateWidgetForStaticProperty<&FMaterialInstanceStaticPropertyGetters::GetAssetPath, &FMaterialInstanceStaticPropertyModifiedGetters::GetAssetPathModified>},
-	{EMaterialInstanceColumnType::Default, "CID_PermutationHash", "Permutation Hash", &SortByStaticProperty<&FMaterialInstanceStaticPropertyGetters::GetPermutationHash>, &GenerateWidgetForStaticProperty<&FMaterialInstanceStaticPropertyGetters::GetPermutationHash, &FMaterialInstanceStaticPropertyModifiedGetters::GetPermutationHashModified> },
-	{EMaterialInstanceColumnType::StaticProperty, "CID_TwoSided", "Two Sided", &SortByStaticProperty<&FMaterialInstanceStaticPropertyGetters::GetTwoSided>, &GenerateWidgetForStaticProperty<&FMaterialInstanceStaticPropertyGetters::GetTwoSided, &FMaterialInstanceStaticPropertyModifiedGetters::GetTwoSidedModified> },
-	{EMaterialInstanceColumnType::StaticProperty, "CID_IsThinSurface", "Thin Surface", &SortByStaticProperty<&FMaterialInstanceStaticPropertyGetters::GetIsThinSurface>, &GenerateWidgetForStaticProperty<&FMaterialInstanceStaticPropertyGetters::GetIsThinSurface, &FMaterialInstanceStaticPropertyModifiedGetters::GetIsThinSurfaceModified> },
-	{EMaterialInstanceColumnType::StaticProperty, "CID_DitheredLODTransition", "Dithered LOD Transition", &SortByStaticProperty<&FMaterialInstanceStaticPropertyGetters::GetDitheredLODTransition>, &GenerateWidgetForStaticProperty<&FMaterialInstanceStaticPropertyGetters::GetDitheredLODTransition, &FMaterialInstanceStaticPropertyModifiedGetters::GetDitheredLODTransitionModified> },
-	{EMaterialInstanceColumnType::StaticProperty, "CID_CastDynamicShadowAsMasked", "Dynamic Shadow As Masked", &SortByStaticProperty<&FMaterialInstanceStaticPropertyGetters::GetCastDynamicShadowAsMasked>, &GenerateWidgetForStaticProperty<&FMaterialInstanceStaticPropertyGetters::GetCastDynamicShadowAsMasked, &FMaterialInstanceStaticPropertyModifiedGetters::GetCastDynamicShadowAsMaskedModified> },
-	{EMaterialInstanceColumnType::StaticProperty, "CID_OutputTranslucentVelocity", "Output Translucent Velocity", &SortByStaticProperty<&FMaterialInstanceStaticPropertyGetters::GetOutputTranslucentVelocity>, &GenerateWidgetForStaticProperty<&FMaterialInstanceStaticPropertyGetters::GetOutputTranslucentVelocity, &FMaterialInstanceStaticPropertyModifiedGetters::GetOutputTranslucentVelocityModified> },
-	{EMaterialInstanceColumnType::StaticProperty, "CID_HasPixelAnimation", "Has Pixel Animation", &SortByStaticProperty<&FMaterialInstanceStaticPropertyGetters::GetHasPixelAnimation>, &GenerateWidgetForStaticProperty<&FMaterialInstanceStaticPropertyGetters::GetHasPixelAnimation, &FMaterialInstanceStaticPropertyModifiedGetters::GetHasPixelAnimationModified> },
-	{EMaterialInstanceColumnType::StaticProperty, "CID_EnableTessellation", "Tessellation", &SortByStaticProperty<&FMaterialInstanceStaticPropertyGetters::GetEnableTessellation>, &GenerateWidgetForStaticProperty<&FMaterialInstanceStaticPropertyGetters::GetEnableTessellation, &FMaterialInstanceStaticPropertyModifiedGetters::GetEnableTessellationModified> },
-	{EMaterialInstanceColumnType::StaticProperty, "CID_BlendMode", "Blend Mode", &SortByStaticProperty<&FMaterialInstanceStaticPropertyGetters::GetBlendMode>, &GenerateWidgetForStaticProperty<&FMaterialInstanceStaticPropertyGetters::GetBlendMode, &FMaterialInstanceStaticPropertyModifiedGetters::GetBlendModeModified> },
-	{EMaterialInstanceColumnType::StaticProperty, "CID_ShadingModel", "Shading Model", &SortByStaticProperty<&FMaterialInstanceStaticPropertyGetters::GetShadingModel>, &GenerateWidgetForStaticProperty<&FMaterialInstanceStaticPropertyGetters::GetShadingModel, &FMaterialInstanceStaticPropertyModifiedGetters::GetShadingModelModified> },
-	{EMaterialInstanceColumnType::StaticProperty, "CID_OpacityMaskClipValue", "Opacity Mask Clip Value", &SortByStaticProperty<&FMaterialInstanceStaticPropertyGetters::GetOpacityMaskClipValue>, &GenerateWidgetForStaticProperty<&FMaterialInstanceStaticPropertyGetters::GetOpacityMaskClipValue, &FMaterialInstanceStaticPropertyModifiedGetters::GetOpacityMaskClipModified> },
+	{EMaterialPropertyCategory::Default, "CID_AssetPath", "Asset Path", &SortByStaticProperty<&FMaterialInstanceStaticPropertyGetters::GetAssetPathIndex>, &GenerateWidgetForStaticProperty<&FMaterialInstanceStaticPropertyGetters::GetAssetPath, &FMaterialInstanceStaticPropertyModifiedGetters::GetAssetPathModified>},
+	{EMaterialPropertyCategory::Default, "CID_PermutationHash", "Permutation Hash", &SortByStaticProperty<&FMaterialInstanceStaticPropertyGetters::GetPermutationHash>, &GenerateWidgetForStaticProperty<&FMaterialInstanceStaticPropertyGetters::GetPermutationHash, &FMaterialInstanceStaticPropertyModifiedGetters::GetPermutationHashModified> },
+	{EMaterialPropertyCategory::StaticProperty, "CID_TwoSided", "Two Sided", &SortByStaticProperty<&FMaterialInstanceStaticPropertyGetters::GetTwoSided>, &GenerateWidgetForStaticProperty<&FMaterialInstanceStaticPropertyGetters::GetTwoSided, &FMaterialInstanceStaticPropertyModifiedGetters::GetTwoSidedModified> },
+	{EMaterialPropertyCategory::StaticProperty, "CID_IsThinSurface", "Thin Surface", &SortByStaticProperty<&FMaterialInstanceStaticPropertyGetters::GetIsThinSurface>, &GenerateWidgetForStaticProperty<&FMaterialInstanceStaticPropertyGetters::GetIsThinSurface, &FMaterialInstanceStaticPropertyModifiedGetters::GetIsThinSurfaceModified> },
+	{EMaterialPropertyCategory::StaticProperty, "CID_DitheredLODTransition", "Dithered LOD Transition", &SortByStaticProperty<&FMaterialInstanceStaticPropertyGetters::GetDitheredLODTransition>, &GenerateWidgetForStaticProperty<&FMaterialInstanceStaticPropertyGetters::GetDitheredLODTransition, &FMaterialInstanceStaticPropertyModifiedGetters::GetDitheredLODTransitionModified> },
+	{EMaterialPropertyCategory::StaticProperty, "CID_CastDynamicShadowAsMasked", "Dynamic Shadow As Masked", &SortByStaticProperty<&FMaterialInstanceStaticPropertyGetters::GetCastDynamicShadowAsMasked>, &GenerateWidgetForStaticProperty<&FMaterialInstanceStaticPropertyGetters::GetCastDynamicShadowAsMasked, &FMaterialInstanceStaticPropertyModifiedGetters::GetCastDynamicShadowAsMaskedModified> },
+	{EMaterialPropertyCategory::StaticProperty, "CID_OutputTranslucentVelocity", "Output Translucent Velocity", &SortByStaticProperty<&FMaterialInstanceStaticPropertyGetters::GetOutputTranslucentVelocity>, &GenerateWidgetForStaticProperty<&FMaterialInstanceStaticPropertyGetters::GetOutputTranslucentVelocity, &FMaterialInstanceStaticPropertyModifiedGetters::GetOutputTranslucentVelocityModified> },
+	{EMaterialPropertyCategory::StaticProperty, "CID_HasPixelAnimation", "Has Pixel Animation", &SortByStaticProperty<&FMaterialInstanceStaticPropertyGetters::GetHasPixelAnimation>, &GenerateWidgetForStaticProperty<&FMaterialInstanceStaticPropertyGetters::GetHasPixelAnimation, &FMaterialInstanceStaticPropertyModifiedGetters::GetHasPixelAnimationModified> },
+	{EMaterialPropertyCategory::StaticProperty, "CID_EnableTessellation", "Tessellation", &SortByStaticProperty<&FMaterialInstanceStaticPropertyGetters::GetEnableTessellation>, &GenerateWidgetForStaticProperty<&FMaterialInstanceStaticPropertyGetters::GetEnableTessellation, &FMaterialInstanceStaticPropertyModifiedGetters::GetEnableTessellationModified> },
+	{EMaterialPropertyCategory::StaticProperty, "CID_BlendMode", "Blend Mode", &SortByStaticProperty<&FMaterialInstanceStaticPropertyGetters::GetBlendMode>, &GenerateWidgetForStaticProperty<&FMaterialInstanceStaticPropertyGetters::GetBlendMode, &FMaterialInstanceStaticPropertyModifiedGetters::GetBlendModeModified> },
+	{EMaterialPropertyCategory::StaticProperty, "CID_ShadingModel", "Shading Model", &SortByStaticProperty<&FMaterialInstanceStaticPropertyGetters::GetShadingModel>, &GenerateWidgetForStaticProperty<&FMaterialInstanceStaticPropertyGetters::GetShadingModel, &FMaterialInstanceStaticPropertyModifiedGetters::GetShadingModelModified> },
+	{EMaterialPropertyCategory::StaticProperty, "CID_OpacityMaskClipValue", "Opacity Mask Clip Value", &SortByStaticProperty<&FMaterialInstanceStaticPropertyGetters::GetOpacityMaskClipValue>, &GenerateWidgetForStaticProperty<&FMaterialInstanceStaticPropertyGetters::GetOpacityMaskClipValue, &FMaterialInstanceStaticPropertyModifiedGetters::GetOpacityMaskClipModified> },
 };
 
 /** Get column FName for a static property enumeration. */
@@ -778,7 +778,7 @@ static const FName UsageColumnIds[MATUSAGE_MAX] = {
 };
 
 /** Static array of column FNames for column type hashes, outside of Default */
-static const FName ColumnTypeHashIds[(int32)(EMaterialInstanceColumnType::Count)] = {
+static const FName CategoryHashIds[(int32)(EMaterialPropertyCategory::Count)] = {
 	"CID_DefaultUnused",
 	"CID_StaticPropertyHash",
 	"CID_UsageFlagHash",
@@ -824,20 +824,20 @@ static void SortByUsage(TArray<TSharedPtr<FMaterialInstanceRowData>>& InArray, E
 /** Sort row data based on the given derived column type hash index. */
 static void SortByColumnHash(TArray<TSharedPtr<FMaterialInstanceRowData>>& InArray, int32 Index, bool bInSortModifer)
 {
-	if (Index < (int32)EMaterialInstanceColumnType::Count)
+	if (Index < (int32)EMaterialPropertyCategory::Count)
 	{
 		if (bInSortModifer)
 		{
 			Algo::StableSort(InArray, [Index](TSharedPtr<FMaterialInstanceRowData> const& A, TSharedPtr<FMaterialInstanceRowData> const& B)
 			{
-				return (A->PerColumnTypeHash[Index]) > (B->PerColumnTypeHash[Index]);
+				return (A->PerCategoryHash[Index]) > (B->PerCategoryHash[Index]);
 			});
 		}
 		else
 		{
 			Algo::StableSort(InArray, [Index](TSharedPtr<FMaterialInstanceRowData> const& A, TSharedPtr<FMaterialInstanceRowData> const& B)
 			{
-				return (A->PerColumnTypeHash[Index]) < (B->PerColumnTypeHash[Index]);
+				return (A->PerCategoryHash[Index]) < (B->PerCategoryHash[Index]);
 			});
 		}
 	}
@@ -981,11 +981,11 @@ public:
 
 			// Try to resolve column as a column type hash
 			static const FNumberFormattingOptions FormatOptions = FNumberFormattingOptions().SetUseGrouping(false);
-			for (int32 Index = 0; Index < (int32)EMaterialInstanceColumnType::Count; ++Index)
+			for (int32 Index = 0; Index < (int32)EMaterialPropertyCategory::Count; ++Index)
 			{
-				if (ColumnId == ColumnTypeHashIds[Index])
+				if (ColumnId == CategoryHashIds[Index])
 				{
-					return SNew(STextBlock).Text(FText::AsNumber(RowItem->PerColumnTypeHash[Index], &FormatOptions));
+					return SNew(STextBlock).Text(FText::AsNumber(RowItem->PerCategoryHash[Index], &FormatOptions));
 				}
 			}
 		}
@@ -1014,8 +1014,8 @@ class SMaterialInstanceTable : public SCompoundWidget
 	/** A slate view of the table. */
 	TSharedPtr<SListView<TSharedPtr<FMaterialInstanceRowData>>> ListView;
 
-	/** Filter for type of columns to show. */
-	EMaterialInstanceColumnType SelectedColumnType;
+	/** Filter for category of columns to show. */
+	EMaterialPropertyCategory SelectedCategory;
 
 	/** Search text used for filtering the rows to show. */
 	FString RowSearchText;
@@ -1045,7 +1045,7 @@ public:
 
 		ParameterNames = InArgs._InitialParameterNames;
 
-		SelectedColumnType = EMaterialInstanceColumnType::StaticProperty;
+		SelectedCategory = EMaterialPropertyCategory::StaticProperty;
 
 		RowSearchText.Reset();
 		ColumnSearchText.Reset();
@@ -1065,11 +1065,11 @@ public:
 				+ SHorizontalBox::Slot().FillWidth(0.2).Padding(2)
 				[
 					SNew(SComboButton)
-					.OnGetMenuContent(this, &SMaterialInstanceTable::OnGetColumnTypeContent)
+					.OnGetMenuContent(this, &SMaterialInstanceTable::OnGetCategoryContent)
 					.ButtonContent()
 					[
 						SNew(STextBlock)
-						.Text(this, &SMaterialInstanceTable::GetSelectedColumnTypeText)
+						.Text(this, &SMaterialInstanceTable::GetSelectedCategoryText)
 					]
 				]
 				+ SHorizontalBox::Slot().FillWidth(0.1).Padding(2)
@@ -1136,24 +1136,24 @@ public:
 	}
 
 protected:
-	/** Called to fill EMaterialInstanceColumnType filter combo box. */
-	TSharedRef<SWidget> OnGetColumnTypeContent()
+	/** Called to fill EMaterialPropertyCategory filter combo box. */
+	TSharedRef<SWidget> OnGetCategoryContent()
 	{
-		const UEnum* Enum = StaticEnum<EMaterialInstanceColumnType>();
+		const UEnum* Enum = StaticEnum<EMaterialPropertyCategory>();
 		FMenuBuilder MenuBuilder(true, NULL);
-		for (int32 Index = 1; Index < (int32)EMaterialInstanceColumnType::Count; ++Index)
+		for (int32 Index = 1; Index < (int32)EMaterialPropertyCategory::Count; ++Index)
 		{
-			FUIAction ItemAction(FExecuteAction::CreateSP(this, &SMaterialInstanceTable::OnSelectColumnType, Index));
+			FUIAction ItemAction(FExecuteAction::CreateSP(this, &SMaterialInstanceTable::OnSelectCategory, Index));
 			MenuBuilder.AddMenuEntry(Enum->GetDisplayNameTextByIndex(Index), TAttribute<FText>(), FSlateIcon(), ItemAction);
 		}
 		return MenuBuilder.MakeWidget();
 	}
 
-	/** Called to get current EMaterialInstanceColumnType filter text entry. */
-	FText GetSelectedColumnTypeText() const
+	/** Called to get current EMaterialPropertyCategory filter text entry. */
+	FText GetSelectedCategoryText() const
 	{
-		const UEnum* Enum = StaticEnum<EMaterialInstanceColumnType>();
-		return Enum->GetDisplayNameTextByIndex((int32)SelectedColumnType);
+		const UEnum* Enum = StaticEnum<EMaterialPropertyCategory>();
+		return Enum->GetDisplayNameTextByIndex((int32)SelectedCategory);
 	}
 	
 	/** Update the view. Must be called after any work that updates filtered/sorted content. */
@@ -1202,10 +1202,10 @@ protected:
 		HeaderRow->AddColumns(AllArgs);
 	}
 
-	/** Called on selection from EMaterialInstanceColumnType filter combo box. */
-	void OnSelectColumnType(int32 Index)
+	/** Called on selection from EMaterialPropertyCategory filter combo box. */
+	void OnSelectCategory(int32 Index)
 	{
-		SelectedColumnType = (EMaterialInstanceColumnType)Index;
+		SelectedCategory = (EMaterialPropertyCategory)Index;
 
 		UpdateRowFilter();
 		UpdateColumnFilter();
@@ -1249,7 +1249,7 @@ protected:
 				if (bUniqueInstances)
 				{
 					bool bInSet = false;
-					PermutationHashesSeen.Add(Row->PerColumnTypeHash[(uint32)SelectedColumnType], &bInSet);
+					PermutationHashesSeen.Add(Row->PerCategoryHash[(uint32)SelectedCategory], &bInSet);
 					if (bInSet)
 					{
 						continue;
@@ -1278,13 +1278,13 @@ protected:
 		{
 			const TSharedPtr<FMaterialInstanceColumnData>& Column = AllColumns[Index];
 
-			if (Column->ColumnType == EMaterialInstanceColumnType::Default)
+			if (Column->Category == EMaterialPropertyCategory::Default)
 			{
 				// Default columns are never hidden.
 				continue;
 			}
 				
-			if (Column->ColumnType == SelectedColumnType && 
+			if (Column->Category == SelectedCategory && 
 				(ColumnSearchText.IsEmpty() || Column->Label.ToString().Contains(ColumnSearchText, ESearchCase::IgnoreCase)))
 			{
 				continue;
@@ -1352,31 +1352,31 @@ protected:
 
 		// Call specific sorting function for the column type.
 		const bool bSortModifier = SortMode == EColumnSortMode::Descending;
-		const EMaterialInstanceColumnType ColumnType = (*SortColumnData)->ColumnType;
+		const EMaterialPropertyCategory Category = (*SortColumnData)->Category;
 
 		// See if we are sorting from a column type hash
-		for (int32 Index = 0; Index < (int32)EMaterialInstanceColumnType::Count; ++Index)
+		for (int32 Index = 0; Index < (int32)EMaterialPropertyCategory::Count; ++Index)
 		{
-			if (ColumnTypeHashIds[Index] == SortColumnId)
+			if (CategoryHashIds[Index] == SortColumnId)
 			{
 				SortByColumnHash(FilteredRows, Index, bSortModifier);
 				return;
 			}
 		}
 
-		switch (ColumnType)
+		switch (Category)
 		{
-			case EMaterialInstanceColumnType::Default:
-			case EMaterialInstanceColumnType::StaticProperty:
+			case EMaterialPropertyCategory::Default:
+			case EMaterialPropertyCategory::StaticProperty:
 				SortByStaticProperty(FilteredRows, GetStaticPropertyFromColumnId(SortColumnId), bSortModifier);
 				break;
-			case EMaterialInstanceColumnType::UsageFlag:
+			case EMaterialPropertyCategory::UsageFlag:
 				SortByUsage(FilteredRows, GetMaterialUsageFromColumnId(SortColumnId), bSortModifier);
 				break;
-			case EMaterialInstanceColumnType::StaticSwitch:
+			case EMaterialPropertyCategory::StaticSwitch:
 				SortByStaticSwitch(FilteredRows, GetStaticSwitchIndexFromColumnId(SortColumnId, *ParameterNames), bSortModifier);
 			break;
-			case EMaterialInstanceColumnType::ComponentMask:
+			case EMaterialPropertyCategory::ComponentMask:
 				SortByComponentMask(FilteredRows, GetComponentMaskIndexFromColumnId(SortColumnId, *ParameterNames), bSortModifier);
 				break;
 		}
@@ -1467,31 +1467,31 @@ static void UpdateMaterialInstanceTables(
 		// Build path and permutation hash so they are first
 		for (int32 PropertyIndex = 0; PropertyIndex < 2; ++PropertyIndex)
 		{
-			Columns.Add(MakeShared<FMaterialInstanceColumnData>(StaticPropertyInfos[PropertyIndex].ColumnId, FText::FromString(StaticPropertyInfos[PropertyIndex].PropertyName), StaticPropertyInfos[PropertyIndex].Type));
+			Columns.Add(MakeShared<FMaterialInstanceColumnData>(StaticPropertyInfos[PropertyIndex].ColumnId, FText::FromString(StaticPropertyInfos[PropertyIndex].PropertyName), StaticPropertyInfos[PropertyIndex].Category));
 		}
 
 		// Build additional type-specific hash columns
-		Columns.Add(MakeShared<FMaterialInstanceColumnData>(ColumnTypeHashIds[(int32)EMaterialInstanceColumnType::StaticProperty], LOCTEXT("Tookkit_StaticPropertyHash","Static Property Hash"), EMaterialInstanceColumnType::StaticProperty));
-		Columns.Add(MakeShared<FMaterialInstanceColumnData>(ColumnTypeHashIds[(int32)EMaterialInstanceColumnType::UsageFlag], LOCTEXT("Tookkit_UsageHash","Usage Hash"), EMaterialInstanceColumnType::UsageFlag));
-		Columns.Add(MakeShared<FMaterialInstanceColumnData>(ColumnTypeHashIds[(int32)EMaterialInstanceColumnType::StaticSwitch], LOCTEXT("Tookkit_StaticSwitchHash","Static Switch Hash"), EMaterialInstanceColumnType::StaticSwitch));
-		Columns.Add(MakeShared<FMaterialInstanceColumnData>(ColumnTypeHashIds[(int32)EMaterialInstanceColumnType::ComponentMask], LOCTEXT("Tookkit_ComponentMaskHash","Component Mask Hash"), EMaterialInstanceColumnType::ComponentMask));
+		Columns.Add(MakeShared<FMaterialInstanceColumnData>(CategoryHashIds[(int32)EMaterialPropertyCategory::StaticProperty], LOCTEXT("Tookkit_StaticPropertyHash","Static Property Hash"), EMaterialPropertyCategory::StaticProperty));
+		Columns.Add(MakeShared<FMaterialInstanceColumnData>(CategoryHashIds[(int32)EMaterialPropertyCategory::UsageFlag], LOCTEXT("Tookkit_UsageHash","Usage Hash"), EMaterialPropertyCategory::UsageFlag));
+		Columns.Add(MakeShared<FMaterialInstanceColumnData>(CategoryHashIds[(int32)EMaterialPropertyCategory::StaticSwitch], LOCTEXT("Tookkit_StaticSwitchHash","Static Switch Hash"), EMaterialPropertyCategory::StaticSwitch));
+		Columns.Add(MakeShared<FMaterialInstanceColumnData>(CategoryHashIds[(int32)EMaterialPropertyCategory::ComponentMask], LOCTEXT("Tookkit_ComponentMaskHash","Component Mask Hash"), EMaterialPropertyCategory::ComponentMask));
 
 		// Build the rest of the columns
 		for (int32 PropertyIndex = 2; PropertyIndex < (int32)EMaterialInstanceStaticProperties::Count; ++PropertyIndex)
 		{
-			Columns.Add(MakeShared<FMaterialInstanceColumnData>(StaticPropertyInfos[PropertyIndex].ColumnId, FText::FromString(StaticPropertyInfos[PropertyIndex].PropertyName), StaticPropertyInfos[PropertyIndex].Type));
+			Columns.Add(MakeShared<FMaterialInstanceColumnData>(StaticPropertyInfos[PropertyIndex].ColumnId, FText::FromString(StaticPropertyInfos[PropertyIndex].PropertyName), StaticPropertyInfos[PropertyIndex].Category));
 		}
 		for (int32 UsageIndex = 0; UsageIndex < MATUSAGE_MAX; ++UsageIndex)
 		{
-			Columns.Add(MakeShared<FMaterialInstanceColumnData>(UsageColumnIds[UsageIndex], FText::FromString(UMaterialInterface::GetUsageName((EMaterialUsage)UsageIndex)), EMaterialInstanceColumnType::UsageFlag));
+			Columns.Add(MakeShared<FMaterialInstanceColumnData>(UsageColumnIds[UsageIndex], FText::FromString(UMaterialInterface::GetUsageName((EMaterialUsage)UsageIndex)), EMaterialPropertyCategory::UsageFlag));
 		}
 		for (FName& Name : AssetInfo.StaticSwitchNames)
 		{
-			Columns.Add(MakeShared<FMaterialInstanceColumnData>(Name, FText::FromName(Name), EMaterialInstanceColumnType::StaticSwitch));
+			Columns.Add(MakeShared<FMaterialInstanceColumnData>(Name, FText::FromName(Name), EMaterialPropertyCategory::StaticSwitch));
 		}
 		for (FName& Name : AssetInfo.ComponentMaskNames)
 		{
-			Columns.Add(MakeShared<FMaterialInstanceColumnData>(Name, FText::FromName(Name), EMaterialInstanceColumnType::ComponentMask));
+			Columns.Add(MakeShared<FMaterialInstanceColumnData>(Name, FText::FromName(Name), EMaterialPropertyCategory::ComponentMask));
 		}
 	}
 	
