MaterialValidation: Add Save Asset and Browse to Asset toolbar buttons
#rnx
#Synced-CL 51620377

diff --git a/Engine/Plugins/Editor/MaterialValidation/Source/Private/MaterialValidationToolkit.cpp b/Engine/Plugins/Editor/MaterialValidation/Source/Private/MaterialValidationToolkit.cpp
--- a/Engine/Plugins/Editor/MaterialValidation/Source/Private/MaterialValidationToolkit.cpp
+++ b/Engine/Plugins/Editor/MaterialValidation/Source/Private/MaterialValidationToolkit.cpp
@@ -1766,7 +1766,7 @@ public:
 				)
 			);
 
-		InitAssetEditor(Mode, InitToolkitHost, GetToolkitFName(), Layout, true /*bCreateDefaultStandaloneMenu*/, false /*bCreateDefaultToolbar*/, EditGroup);
+		InitAssetEditor(Mode, InitToolkitHost, GetToolkitFName(), Layout, true /*bCreateDefaultStandaloneMenu*/, true /*bCreateDefaultToolbar*/, EditGroup);
 	}
 
 	/** Apply selected material to the material instances table. */
