MaterialValidation: Fix incorrect modified shader count in some cases when a material hierarchy has instances without static permutations.
In that case we need to inherit the material layer hash correctly.
#rnx
#Synced-CL 51666760

diff --git a/Engine/Plugins/Editor/MaterialValidation/Source/Private/MaterialValidationLibrary.cpp b/Engine/Plugins/Editor/MaterialValidation/Source/Private/MaterialValidationLibrary.cpp
--- a/Engine/Plugins/Editor/MaterialValidation/Source/Private/MaterialValidationLibrary.cpp
+++ b/Engine/Plugins/Editor/MaterialValidation/Source/Private/MaterialValidationLibrary.cpp
@@ -1244,7 +1244,7 @@ int32 UMaterialValidationLibrary::GetModifiedShaderCount(UMaterialValidationGrou
 		MaterialValidationLibrary::ApplyBitMaskOverrides(Parent.StaticSwitchOverrideValues, AssetData.StaticSwitchOverrideMask, AssetData.StaticSwitchOverrideValues, Child.StaticSwitchOverrideValues, Child.StaticSwitchOverrideMask);
 		MaterialValidationLibrary::ApplyBitMaskOverrides(Parent.ComponentMaskOverrideValues, AssetData.ComponentMaskOverrideMask, AssetData.ComponentMaskOverrideValues, Child.ComponentMaskOverrideValues, Child.ComponentMaskOverrideMask);
 		MaterialValidationLibrary::ApplyBasePropertyOverrides(Parent.StaticProperties, AssetData.StaticPropertyOverrideFlags, AssetData.StaticProperties, Child.StaticProperties, Child.StaticPropertyOverrideFlags);
-		Child.MaterialLayerHash = AssetData.MaterialLayerHash;
+		Child.MaterialLayerHash = AssetData.MaterialLayerHash ? AssetData.MaterialLayerHash : Parent.MaterialLayerHash;
 
 		// Build the permutation hash and if this is a new permutation then add its shader count.
 		Child.PermutationHash = MaterialValidation::BuildPermutationHash(Child.StaticProperties, Child.StaticSwitchOverrideValues, Child.ComponentMaskOverrideValues, Child.MaterialLayerHash);
