MaterialValidationToolkit: Small refactorof some things to shared header to simplify an incoming change.
#rnx
#Synced-CL 51666760

diff --git a/Engine/Plugins/Editor/MaterialValidation/Source/Private/MaterialValidationGroupDiff.cpp b/Engine/Plugins/Editor/MaterialValidation/Source/Private/MaterialValidationGroupDiff.cpp
--- a/Engine/Plugins/Editor/MaterialValidation/Source/Private/MaterialValidationGroupDiff.cpp
+++ b/Engine/Plugins/Editor/MaterialValidation/Source/Private/MaterialValidationGroupDiff.cpp
@@ -8,7 +8,7 @@
 #include "MaterialValidationGroup.h"
 #include "MaterialValidationLibrary.h"
 #include "MaterialValidationLibraryTypes.h"
-#include "MaterialValidationToolkit.h"
+#include "MaterialValidationToolkitShared.h"
 #include "Misc/TextFilterExpressionEvaluator.h"
 #include "PropertyCustomizationHelpers.h"
 #include "Widgets/Input/SButton.h"
@@ -24,32 +24,6 @@
 
 namespace MaterialValidation {
 
-/** Value to track user selected width of asset name columns. */
-static float GAssetNameColumnWidth_GroupDiff = 700.f;
-
-/** Get color to use according to whether cost delta is positive, negative or neutral. */
-static FSlateColor GetDeltaColor(int32 InDelta)
-{
-	if (InDelta > 0)
-	{
-		return FSlateColor(FLinearColor(1.f, 0.4f, 0.4f));
-	}
-	if (InDelta < 0)
-	{
-		return FSlateColor(FLinearColor(0.4f, 1.f, 0.4f));
-	}
-	return FSlateColor::UseForeground();
-}
-
-/** Return FText that formats the number with a positive/negative sign. */
-static FText GetTextSignedInt(int32 Value)
-{
-	return (Value > 0) ? FText::Format(LOCTEXT("PositiveNumberFormat", "+{0}"),	FText::AsNumber(Value)) : FText::AsNumber(Value);
-}
-
-static const FText GEmptyText = FText::FromString("-");
-static const FText GUnknownText = FText::FromString("?");
-
 /** Per-material row data for the material diff table. Holds the counts and deltas from both snapshots. */
 struct FMaterialDiffStats
 {
@@ -315,14 +289,14 @@ private:
 				.HAlignHeader(HAlign_Center).HAlignCell(HAlign_Left)
 				.DefaultLabel(LOCTEXT("DiffAssetPath", "Asset Path"))
 				.OnSort(this, &SMaterialDiffTable::OnColumnSortModeChanged)
-				.ManualWidth(GAssetNameColumnWidth_GroupDiff)
+				.ManualWidth(GAssetNameColumnWidth)
 				.ManualWidth_Lambda([this]() 
 				{
-					return GAssetNameColumnWidth_GroupDiff;
+					return GAssetNameColumnWidth;
 				})
 				.OnWidthChanged_Lambda([this](float NewWidth) 
 				{
-					GAssetNameColumnWidth_GroupDiff = NewWidth;
+					GAssetNameColumnWidth = NewWidth;
 				})
 			+ MakeSortableCol(FMaterialDiffStats::ColumnId_ShaderDelta, LOCTEXT("ShaderDelta", "Shaders (Delta)"))
 			+ MakeSortableCol(FMaterialDiffStats::ColumnId_InstanceDelta, LOCTEXT("InstanceDelta", "Instances (Delta)"))
@@ -635,19 +609,19 @@ public:
 						SNew(SHeaderRow)
 						+ SHeaderRow::Column(SMaterialInstanceDiffTreeRow::ColAssetPath)
 							.HAlignHeader(HAlign_Center)
-							.ManualWidth(GAssetNameColumnWidth_GroupDiff)
+							.ManualWidth(GAssetNameColumnWidth)
 							.ManualWidth_Lambda([this]()
 							{
-								return GAssetNameColumnWidth_GroupDiff;
+								return GAssetNameColumnWidth;
 							})
 							.OnWidthChanged_Lambda([this](float NewWidth)
 							{
-								GAssetNameColumnWidth_GroupDiff = NewWidth;
+								GAssetNameColumnWidth = NewWidth;
 							})
 							.HeaderContent()
 							[
 								SNew(SBox)
-								.WidthOverride(TAttribute<FOptionalSize>::CreateLambda([]() -> FOptionalSize { return GAssetNameColumnWidth_GroupDiff; }))
+								.WidthOverride(TAttribute<FOptionalSize>::CreateLambda([]() -> FOptionalSize { return GAssetNameColumnWidth; }))
 								[
 									SNew(SButton)
 									.ButtonStyle(FAppStyle::Get(), "NoBorder")
diff --git a/Engine/Plugins/Editor/MaterialValidation/Source/Private/MaterialValidationToolkit.cpp b/Engine/Plugins/Editor/MaterialValidation/Source/Private/MaterialValidationToolkit.cpp
--- a/Engine/Plugins/Editor/MaterialValidation/Source/Private/MaterialValidationToolkit.cpp
+++ b/Engine/Plugins/Editor/MaterialValidation/Source/Private/MaterialValidationToolkit.cpp
@@ -3,13 +3,14 @@
 #include "MaterialValidationToolkit.h"
 
 #include "AssetRegistry/IAssetRegistry.h"
-#include "Hash/xxhash.h"
 #include "Framework/Application/SlateApplication.h"
+#include "Hash/xxhash.h"
 #include "Materials/MaterialInterface.h"
 #include "MaterialValidationGroup.h"
 #include "MaterialValidationGroupDiff.h"
 #include "MaterialValidationLibrary.h"
 #include "MaterialValidationLibraryTypes.h"
+#include "MaterialValidationToolkitShared.h"
 #include "Misc/NotifyHook.h"
 #include "Misc/TextFilterExpressionEvaluator.h"
 #include "Modules/ModuleManager.h"
@@ -30,8 +31,74 @@
 
 namespace MaterialValidation {
 
-/** Global shared value to track user selected width of asset name columns. */
-static float GAssetNameColumnWidth = 700.f;
+float GAssetNameColumnWidth = 700.f;
 
+const FText GEmptyText = FText::FromString("-");
+const FText GUnknownText = FText::FromString("?");
+
+FSlateColor GetDeltaColor(int32 InDelta)
+{
+	if (InDelta > 0)
+	{
+		return FSlateColor(FLinearColor(1.f, 0.4f, 0.4f));
+	}
+	else if (InDelta < 0)
+	{
+		return FSlateColor(FLinearColor(0.4f, 1.f, 0.4f));
+	}
+	return FSlateColor::UseForeground();
+}
+
+FText GetTextSignedInt(int32 Value)
+{
+	return (Value > 0) ? FText::Format(LOCTEXT("PositiveNumberFormat", "+{0}"), FText::AsNumber(Value)) : FText::AsNumber(Value);
+}
+
+TSharedRef<SWidget> GenerateAssetPathWidget(FSoftObjectPath InAssetPath, FSlateColor InColor)
+{
+	const FAssetData AssetData = IAssetRegistry::Get()->GetAssetByObjectPath(InAssetPath);
+	const EVisibility BrowseVisibility = AssetData.IsValid() ? EVisibility::Visible : EVisibility::Hidden;
+
+	return SNew(SBox)
+		.HAlign(HAlign_Fill)
+		.VAlign(VAlign_Center)
+		[
+			SNew(SHorizontalBox)
+				+ SHorizontalBox::Slot()
+				.FillWidth(1.0f)
+				.VAlign(VAlign_Center)
+				[
+					SNew(STextBlock)
+						.Text(FText::FromString(InAssetPath.ToString()))
+						.ToolTipText(FText::FromString(InAssetPath.ToString()))
+						.ColorAndOpacity(InColor)
+				]
+				+ SHorizontalBox::Slot()
+				.AutoWidth()
+				.VAlign(VAlign_Center)
+				.Padding(4.0f, 0.0f)
+				[
+					SNew(SBox)
+						.Visibility(BrowseVisibility)
+						[
+							PropertyCustomizationHelpers::MakeBrowseButton(
+								FSimpleDelegate::CreateLambda([AssetData]()
+									{
+										const FModifierKeysState KeyState = FSlateApplication::Get().GetModifierKeys();
+										if (KeyState.IsAltDown())
+										{
+											GEditor->EditObject(AssetData.GetAsset());
+										}
+										else
+										{
+											GEditor->SyncBrowserToObjects({ AssetData });
+										}
+									}),
+								LOCTEXT("BrowseTooltip", "Browse to Asset in Content Browser"))
+						]
+				]
+		];
+}
+
 /** A data table entry describing a UMaterial. */
 struct FMaterialStats
@@ -183,53 +250,6 @@
 	return SNew(SExperimentsTable);
 }
 
-/** Helper function to create an SWidget for AssetPath with text and a browse button. */
-TSharedRef<SWidget> GenerateAssetPathWidget(FSoftObjectPath InAssetPath, FSlateColor InColor)
-{
-	const FAssetData AssetData = IAssetRegistry::Get()->GetAssetByObjectPath(InAssetPath);
-	const EVisibility BrowseVisibility = AssetData.IsValid() ? EVisibility::Visible : EVisibility::Hidden;
-
-	return SNew(SBox)
-		.HAlign(HAlign_Fill)
-		.VAlign(VAlign_Center)
-		[
-			SNew(SHorizontalBox)
-				+ SHorizontalBox::Slot()
-				.FillWidth(1.0f)
-				.VAlign(VAlign_Center)
-				[
-					SNew(STextBlock)
-						.Text(FText::FromString(InAssetPath.ToString()))
-						.ToolTipText(FText::FromString(InAssetPath.ToString()))
-						.ColorAndOpacity(InColor)
-				]
-				+ SHorizontalBox::Slot()
-				.AutoWidth()
-				.VAlign(VAlign_Center)
-				.Padding(4.0f, 0.0f)
-				[
-					SNew(SBox)
-					.Visibility(BrowseVisibility)
-					[
-						PropertyCustomizationHelpers::MakeBrowseButton(
-							FSimpleDelegate::CreateLambda([AssetData]()
-							{
-								const FModifierKeysState KeyState = FSlateApplication::Get().GetModifierKeys();
-								if (KeyState.IsAltDown())
-								{
-									GEditor->EditObject(AssetData.GetAsset());
-								}
-								else
-								{
-									GEditor->SyncBrowserToObjects({ AssetData });
-								}
-							}),
-							LOCTEXT("BrowseTooltip", "Browse to Asset in Content Browser"))
-					]
-				]
-		];
-}
-
 /** A slate container of a data table entry describing a UMaterial. */
 class SMaterialTableRow : public SMultiColumnTableRow<TSharedPtr<FMaterialStats>>
 {
@@ -1598,8 +1618,8 @@
 
 } // namespace MaterialValidation
 
-/** 
- * AssetEditorToolkit for UMaterialValidationGroup. 
+/**
+ * AssetEditorToolkit for UMaterialValidationGroup.
  * This sets up the tabs for showing various table views of the data in the group.
  */
 class FMaterialValidationToolkit : public FAssetEditorToolkit, public FGCObject, public FNotifyHook
@@ -1610,14 +1630,14 @@
 	/** Menu category for the toolkit tabs. */
 	TSharedPtr<class FWorkspaceItem> WorkspaceMenuCategory;
 
-	TSharedPtr<MaterialValidation::SExperimentsTable> ExperimentsTable;
-
 	/** Standard details panel view. */
-	TSharedPtr<class IDetailsView> DetailsView;
+	TSharedPtr<IDetailsView> DetailsView;
 	/** Material table widget. */
-	TSharedPtr<MaterialValidation::SMaterialTable> MaterialTable;
+	TSharedPtr<SWidget> MaterialTable;
 	/** Material instance table widget. */
 	TSharedPtr<MaterialValidation::SMaterialInstanceTable> MaterialInstanceTable;
+	/** Experiments table widget. */
+	TSharedPtr<MaterialValidation::SExperimentsTable> ExperimentsTable;
 
 	// Toolkit tabs are referenced by FName.
 	static const FName TabId_Details;
diff --git a/Engine/Plugins/Editor/MaterialValidation/Source/Private/MaterialValidationToolkit.h b/Engine/Plugins/Editor/MaterialValidation/Source/Private/MaterialValidationToolkit.h
--- a/Engine/Plugins/Editor/MaterialValidation/Source/Private/MaterialValidationToolkit.h
+++ b/Engine/Plugins/Editor/MaterialValidation/Source/Private/MaterialValidationToolkit.h
@@ -18,8 +18,3 @@ class UAssetDefinition_MaterialValidationGroup : public UAssetDefinitionDefault
 	virtual EAssetCommandResult OpenAssets(const FAssetOpenArgs& OpenArgs) const override;
 	virtual EAssetCommandResult PerformAssetDiff(const FAssetDiffArgs& DiffArgs) const override;
 };
-
-namespace MaterialValidation
-{
-	TSharedRef<SWidget> GenerateAssetPathWidget(FSoftObjectPath InAssetPath, FSlateColor InColor);
-}
diff --git a/Engine/Plugins/Editor/MaterialValidation/Source/Private/MaterialValidationToolkitShared.h b/Engine/Plugins/Editor/MaterialValidation/Source/Private/MaterialValidationToolkitShared.h
new file mode 100644
--- /dev/null
+++ b/Engine/Plugins/Editor/MaterialValidation/Source/Private/MaterialValidationToolkitShared.h
@@ -0,0 +1,24 @@
+// Copyright Epic Games, Inc. All Rights Reserved.
+#pragma once
+
+#include "CoreMinimal.h"
+#include "Styling/SlateColor.h"
+
+class SWidget;
+
+namespace MaterialValidation
+{
+	/** Shared column width for asset-path columns across all tables in the toolkit. */
+	extern float GAssetNameColumnWidth;
+	/** Text to use for empty cells. */
+	extern const FText GEmptyText;
+	/** Text to use for unknown values. */
+	extern const FText GUnknownText;
+
+	/** Creates a path display widget with a browse-to-asset button. */
+	TSharedRef<SWidget> GenerateAssetPathWidget(FSoftObjectPath InAssetPath, FSlateColor InColor);
+	/** Get color to use according to whether cost delta is positive, negative or neutral. */
+	FSlateColor GetDeltaColor(int32 InDelta);
+	/** Return FText that formats the number with a positive/negative sign. */
+	FText GetTextSignedInt(int32 Value);
+}
