MaterialValidation: Support material replacement when evaluating shader counts.
This will be useful for experiments where we want to use replacement materials and measure the effect.
#rnx
#Synced-CL 51666760

diff --git a/Engine/Plugins/Editor/MaterialValidation/Source/Private/MaterialValidationLibrary.cpp b/Engine/Plugins/Editor/MaterialValidation/Source/Private/MaterialValidationLibrary.cpp
--- a/Engine/Plugins/Editor/MaterialValidation/Source/Private/MaterialValidationLibrary.cpp
+++ b/Engine/Plugins/Editor/MaterialValidation/Source/Private/MaterialValidationLibrary.cpp
@@ -1141,18 +1141,45 @@ int32 UMaterialValidationLibrary::GetShaderCount(UMaterialValidationGroup const*
 	return MaterialDesc != nullptr ? MaterialDesc->NumShadersTotal : 0;
 }
 
-int32 UMaterialValidationLibrary::GetModifiedShaderCount(UMaterialValidationGroup const* InGroup, UMaterial* InBaseMaterial, TArray<UMaterialInterface*>& InModifiedObjects)
+int32 UMaterialValidationLibrary::GetModifiedShaderCount(
+	UMaterialValidationGroup const* InGroup, 
+	UMaterial* InBaseMaterial, 
+	TArray<UMaterialInterface*> const& InModifiedObjects,
+	TArray<UMaterialInterface*> const& InReplacementObjects,
+	bool bForceLoadObjects)
 {
 	if (InGroup == nullptr || InBaseMaterial == nullptr)
 	{
 		return 0;
 	}
 
-	// Get paths to modified objects for lookup.
-	TMap<FSoftObjectPath, UMaterialInterface*> ModifiedObjectPaths;
-	for (UMaterialInterface* Object : InModifiedObjects)
+	// Build a map from the canonical object path used in the database to the modified object entries.
+	TMap<FSoftObjectPath, int32> ModifiedObjectPaths;
+	for (int32 Index = 0; Index < InModifiedObjects.Num(); ++Index)
 	{
-		ModifiedObjectPaths.Add(FSoftObjectPath(Object), Object);
+		ModifiedObjectPaths.Add(FSoftObjectPath(InModifiedObjects[Index]), Index);
+	}
+	// Helper to find the modified entry from the path.
+	auto GetModifiedObject = [&](FSoftObjectPath const& InPath) -> UMaterialInterface*
+	{
+		if (int32 const* FoundIndex = ModifiedObjectPaths.Find(InPath))
+		{
+			UMaterialInterface* ModifiedObject = InModifiedObjects.IsValidIndex(*FoundIndex) ? InModifiedObjects[*FoundIndex] : nullptr;
+			UMaterialInterface* ReplacementObject = InReplacementObjects.IsValidIndex(*FoundIndex) ? InReplacementObjects[*FoundIndex] : nullptr;
+			return ReplacementObject != nullptr ? ReplacementObject : ModifiedObject;
+		}
+		return nullptr;
+	};
+
+	UMaterial* BaseMaterial = InBaseMaterial;
+
+	// Find any replacement for the base material.
+	const int32 ModifiedObjectIndex = InModifiedObjects.Find(InBaseMaterial);
+	const bool bIsBaseMaterialInModified = ModifiedObjectIndex != INDEX_NONE;
+	if (InReplacementObjects.IsValidIndex(ModifiedObjectIndex))
+	{
+		UMaterial* Replacement = Cast<UMaterial>(InReplacementObjects[ModifiedObjectIndex]);
+		BaseMaterial = Replacement != nullptr ? Replacement : BaseMaterial;
 	}
 
 	TArray<FString> SearchPaths;
@@ -1160,18 +1187,18 @@ int32 UMaterialValidationLibrary::GetModifiedShaderCount(UMaterialValidationGrou
 	MaterialValidationLibrary::GetSearchPaths(*InGroup, SearchPaths, ExcludePaths);
 
 	MaterialValidationLibrary::FMaterialPermutationData PermutationData;
-	BuildPermutationData(*InBaseMaterial, PermutationData);
+	BuildPermutationData(*BaseMaterial, PermutationData);
 
 	FMaterialValidationAssetData_Material MaterialAssetDataFromObject;
-	GetAssetData(*InBaseMaterial, PermutationData, MaterialAssetDataFromObject);
+	GetAssetData(*BaseMaterial, PermutationData, MaterialAssetDataFromObject);
 
 	// Get the material asset data stored on the group.
 	// But evaluate from the material if the material is modified or if it is not found in the group.
-	const FSoftObjectPath  BaseMaterialSoftObjectPath(InBaseMaterial);
+	const FSoftObjectPath BaseMaterialSoftObjectPath(InBaseMaterial);
 	FMaterialValidationDesc const* MaterialDesc = MaterialValidationLibrary::FindMaterialDescFromSoftObjectPath(*InGroup, BaseMaterialSoftObjectPath);
 
 	FMaterialValidationAssetData_Material MaterialAssetData;
-	if (MaterialDesc == nullptr || InModifiedObjects.Contains(InBaseMaterial))
+	if (MaterialDesc == nullptr || bIsBaseMaterialInModified || bForceLoadObjects)
 	{
 		MaterialAssetData = MaterialAssetDataFromObject;
 	}
@@ -1185,7 +1212,7 @@ int32 UMaterialValidationLibrary::GetModifiedShaderCount(UMaterialValidationGrou
 	// So instead we approximate the change by only looking at the shader delta on the base material.
 	if (MaterialDesc != nullptr && MaterialDesc->AssetData.StaticPropertyLayoutHash != MaterialAssetDataFromObject.StaticPropertyLayoutHash)
 	{
-		const int32 NumShadersBase = MaterialValidationLibrary::GetNumShaderTypes(*InBaseMaterial);
+		const int32 NumShadersBase = MaterialValidationLibrary::GetNumShaderTypes(*BaseMaterial);
 		return MaterialDesc->NumShadersTotal + (NumShadersBase - MaterialDesc->NumShadersBase) * (1 + MaterialDesc->PermutationHashes.Num());
 	}
 
@@ -1208,7 +1235,7 @@ int32 UMaterialValidationLibrary::GetModifiedShaderCount(UMaterialValidationGrou
 	}
 
 	// Initialize the shader count and unique hash array.
-	int32 NumShaders = MaterialValidationLibrary::GetNumShaderTypes(*InBaseMaterial);
+	int32 NumShaders = MaterialValidationLibrary::GetNumShaderTypes(*BaseMaterial);
 	TArray<uint32> PermutationHashes;
 	PermutationHashes.Add(MaterialAssetData.PermutationHash);
 
@@ -1218,8 +1245,12 @@ int32 UMaterialValidationLibrary::GetModifiedShaderCount(UMaterialValidationGrou
 		const int32 CurrentDepth = HierarchyDepths[Index];
 		FString PackagePath = MaterialValidationLibrary::AddPackagePathExtension(PackageNames[Index]);
 		FSoftObjectPath SoftObjectPackagePath = MaterialValidationLibrary::ConvertStringPathToSoftObjectPath(PackagePath);
-		UMaterialInterface** FoundModifiedObject = ModifiedObjectPaths.Find(SoftObjectPackagePath);
-		UMaterialInstanceConstant* MaterialInstance = FoundModifiedObject ? Cast<UMaterialInstanceConstant>(*FoundModifiedObject) : nullptr;
+		UMaterialInstanceConstant* MaterialInstance = Cast<UMaterialInstanceConstant>(GetModifiedObject(SoftObjectPackagePath));
+		
+		if (MaterialInstance == nullptr && bForceLoadObjects)
+		{
+			MaterialInstance = TSoftObjectPtr<UMaterialInstanceConstant>(SoftObjectPackagePath).LoadSynchronous();
+		}
 
 		FMaterialValidationAssetData_MaterialInstance AssetData;
 		const bool bValidAssetData = MaterialValidationLibrary::GetMaterialInstanceAssetData(PackagePath, MaterialInstance, MaterialDesc, PermutationData, AssetData);
@@ -1250,13 +1281,7 @@ int32 UMaterialValidationLibrary::GetModifiedShaderCount(UMaterialValidationGrou
 		Child.PermutationHash = MaterialValidation::BuildPermutationHash(Child.StaticProperties, Child.StaticSwitchOverrideValues, Child.ComponentMaskOverrideValues, Child.MaterialLayerHash);
 		if (!PermutationHashes.Contains(Child.PermutationHash))
 		{
-			// If we don't have the material instance loaded then create a transient version.
-			if (MaterialInstance == nullptr)
-			{
-				MaterialInstance = MaterialValidationLibrary::CreateMaterialInstance(*InBaseMaterial, Child);
-			}
-
-			NumShaders += MaterialValidationLibrary::GetNumShaderTypes(*MaterialInstance);
+			NumShaders += MaterialValidationLibrary::GetNumShaderTypes(*MaterialValidationLibrary::CreateMaterialInstance(*BaseMaterial, Child));
 			PermutationHashes.Add(Child.PermutationHash);
 		}
 	}
@@ -1264,8 +1289,8 @@ int32 UMaterialValidationLibrary::GetModifiedShaderCount(UMaterialValidationGrou
 	return NumShaders;
 }
 
-/** 
- * Note there is quite a bit of duplication with GetModifiedShaderCount() in this function. 
+/**
+ * Note there is quite a bit of duplication with GetModifiedShaderCount() in this function.
  * It might be good to consolidate on a single material instance visitor that can do arbitrary work?
  */
 void UMaterialValidationLibrary::GetMaterialHierarchyInfo(UMaterialValidationGroup const* InGroup, FSoftObjectPath const& InBaseMaterialPath, FMaterialDatabaseAssetHierarchyInfo& OutInfo)
diff --git a/Engine/Plugins/Editor/MaterialValidation/Source/Private/MaterialValidationLibrary.h b/Engine/Plugins/Editor/MaterialValidation/Source/Private/MaterialValidationLibrary.h
--- a/Engine/Plugins/Editor/MaterialValidation/Source/Private/MaterialValidationLibrary.h
+++ b/Engine/Plugins/Editor/MaterialValidation/Source/Private/MaterialValidationLibrary.h
@@ -110,17 +110,25 @@ class UMaterialValidationLibrary : public UBlueprintFunctionLibrary
 	/**
 	 * Get the number of shaders used by all permutations of a base material taking into account local changes in the materials.
 	 *
-	 * @param	InGroup				The group to use.
-	 * @param	InMaterial			The base material to evaluate.
-	 * @param	InModifiedObjects	The array of material and material instance objects that need evaluating with their local changes.
-	 * 
-	 * @returns	The total number of shaders. This value is the best estimate that we can give without triggering asset loads.
+	 * @param	InGroup					The group to use.
+	 * @param	InMaterial				The base material to evaluate.
+	 * @param	InModifiedObjects		An array of loaded material objects that should be used instead of the corresponding stored database values when evaluating permutations.
+	 * @param	InReplacementObjects	An array of loaded material objects in sync with InModifiedObjects. When the corresponding array item exists and is valid it is used to replace 
+	 *									the object in InModifiedObjects when evaluation permutations. This can be used to evaluate the impact of using a different object.
+	 * @param	bForceLoadObjects		When true, any hierarchy asset not already in memory is synchronously loaded from disk.
+	 *
+	 * @returns	The total number of shaders. This value is the best estimate that we can give without triggering asset loads if bForceLoad=false.
 	 */
-	UFUNCTION(BlueprintCallable, Category = Material)
-	static int32 GetModifiedShaderCount(UMaterialValidationGroup const* InGroup, UMaterial* InMaterial, TArray<UMaterialInterface*>& InModifiedObjects);
+	UFUNCTION(BlueprintCallable, Category = Material, meta = (AutoCreateRefTerm = "InReplacementObjects"))
+	static int32 GetModifiedShaderCount(
+		UMaterialValidationGroup const* InGroup, 
+		UMaterial* InMaterial, 
+		TArray<UMaterialInterface*> const& InModifiedObjects,
+		TArray<UMaterialInterface*> const& InReplacementObjects,
+		bool bForceLoadObjects = false);
 
-	/** 
-	 * Get all of the material instance data for a single material hierarchy. 
+	/**
+	 * Get all of the material instance data for a single material hierarchy.
 	 *
 	 * @param	InGroup				The group to use.
 	 * @param	InBaseMaterialPath	The path of the base material to evaluate.
diff --git a/Engine/Plugins/Editor/MaterialValidation/Source/Private/MaterialValidators.cpp b/Engine/Plugins/Editor/MaterialValidation/Source/Private/MaterialValidators.cpp
--- a/Engine/Plugins/Editor/MaterialValidation/Source/Private/MaterialValidators.cpp
+++ b/Engine/Plugins/Editor/MaterialValidation/Source/Private/MaterialValidators.cpp
@@ -141,7 +141,7 @@ void UEditorValidator_MaterialPermutation::PostAssetValidation(TArray<TSharedRef
 				if (bIsInGroupPath)
 				{
 					NumShaders -= UMaterialValidationLibrary::GetShaderCount(Group, BaseMaterial);
-					NumShaders += UMaterialValidationLibrary::GetModifiedShaderCount(Group, BaseMaterial, MutableView(ModifiedObjects));
+					NumShaders += UMaterialValidationLibrary::GetModifiedShaderCount(Group, BaseMaterial, MutableView(ModifiedObjects), {}, false /*bForceLoadObjects*/);
 				}
 			}
 		}
