MaterialValidation: Add accurate slow validation fallback to PostAssetValidation shader budget check.
When fast validation (workspace vs stored baseline) hits the submit threshold we now apply an additional slow validation (workspace vs depot head) for a more accurate shader delta.
#rnx
#Synced-CL 51704221

diff --git a/Engine/Plugins/Editor/MaterialValidation/Source/Private/MaterialValidationChangelistTab.cpp b/Engine/Plugins/Editor/MaterialValidation/Source/Private/MaterialValidationChangelistTab.cpp
--- a/Engine/Plugins/Editor/MaterialValidation/Source/Private/MaterialValidationChangelistTab.cpp
+++ b/Engine/Plugins/Editor/MaterialValidation/Source/Private/MaterialValidationChangelistTab.cpp
@@ -4,14 +4,13 @@
 
 #include "AssetRegistry/AssetData.h"
 #include "AssetRegistry/IAssetRegistry.h"
-#include "DiffUtils.h"
 #include "ISourceControlChangelistState.h"
 #include "ISourceControlModule.h"
 #include "ISourceControlProvider.h"
-#include "ISourceControlRevision.h"
 #include "MaterialEditingLibrary.h"
 #include "MaterialValidationGroup.h"
 #include "MaterialValidationLibrary.h"
+#include "MaterialValidators.h"
 #include "MaterialValidationLibraryTypes.h"
 #include "MaterialValidationToolkitShared.h"
 #include "Materials/Material.h"
@@ -459,14 +458,12 @@ private:
 	}
 
 	/**
-	 * Iterates the files in InFileStates and loads workspace assets to fills OutResult.
-	 * When bLoadDepotVersions is true the depot revision of each material is also loaded and stored in ReplacementObjects to enable before/after comparison.
-	 * Material functions are expanded to their referencing base materials in fast mode but in slow mode only bHasFunctionFiles is set.
+	 * Iterates the files in InFileStates and loads workspace assets to fill OutResult.
+	 * When bExpandFunctions is true material function files are expanded to their referencing base materials.
+	 * Depot loading is not performed here; call FetchRevisionHistory + LoadDepotVersions separately for slow mode.
 	 */
-	static void CollectFromChangelist(TArray<UMaterialValidationGroup*> const& InGroups, TArray<FSourceControlStateRef> const& InFileStates, bool bLoadDepotVersions, FValidateCollectionResult& OutResult)
+	static void CollectFromChangelist(TArray<UMaterialValidationGroup*> const& InGroups, TArray<FSourceControlStateRef> const& InFileStates, bool bExpandFunctions, FValidateCollectionResult& OutResult)
 	{
-		ISourceControlProvider* Provider = bLoadDepotVersions ? &ISourceControlModule::Get().GetProvider() : nullptr;
-
 		for (FSourceControlStateRef const& FileState : InFileStates)
 		{
 			FString const& Filename = FileState->GetFilename();
@@ -492,62 +489,35 @@ private:
 			UMaterialInstanceConstant* MaterialInstance = Cast<UMaterialInstanceConstant>(WorkspaceAsset);
 			UMaterialFunction* MaterialFunction = Cast<UMaterialFunction>(WorkspaceAsset);
 
-			if (Material || MaterialInstance)
+			if (Material != nullptr)
 			{
-				// Optionally load the depot version as a substitute for the before/after comparison.
-				UMaterialInterface* DepotVersion = nullptr;
-				if (Provider != nullptr)
-				{
-					FSourceControlStatePtr HistoryState = Provider->GetState(Filename, EStateCacheUsage::Use);
-					if (HistoryState.IsValid() && HistoryState->GetHistorySize() > 0)
-					{
-						TSharedPtr<ISourceControlRevision> Revision = HistoryState->GetHistoryItem(0);
-						if (Revision.IsValid())
-						{
-							UPackage* DepotPackage = DiffUtils::LoadPackageForDiff(Revision);
-							if (DepotPackage != nullptr)
-							{
-								OutResult.DepotPackages.Add(DepotPackage);
-								DepotVersion = FindObject<UMaterialInterface>(DepotPackage, *ShortName);
-							}
-						}
-					}
-				}
-
-				if (Material != nullptr)
+				if (FindGroupForMaterial(InGroups, Material))
 				{
-					if (FindGroupForMaterial(InGroups, Material))
-					{
-						OutResult.BaseMaterials.AddUnique(Material);
-						OutResult.ChangelistInstanceCounts.FindOrAdd(FSoftObjectPath(Material))++;
-						OutResult.ModifiedObjects.Add(Material);
-						OutResult.ReplacementObjects.Add(DepotVersion);
-					}
+					OutResult.BaseMaterials.AddUnique(Material);
+					OutResult.ChangelistInstanceCounts.FindOrAdd(FSoftObjectPath(Material))++;
+					OutResult.ModifiedObjects.Add(Material);
 				}
-				if (MaterialInstance != nullptr && MaterialInstance->bHasStaticPermutationResource)
+			}
+			if (MaterialInstance != nullptr && MaterialInstance->bHasStaticPermutationResource)
+			{
+				if (FindGroupForMaterialInstance(InGroups, MaterialInstance))
 				{
-					if (FindGroupForMaterialInstance(InGroups, MaterialInstance))
-					{
-						OutResult.BaseMaterials.AddUnique(MaterialInstance->GetMaterial());
-						OutResult.ChangelistInstanceCounts.FindOrAdd(FSoftObjectPath(MaterialInstance->GetMaterial()))++;
-						OutResult.ModifiedObjects.Add(MaterialInstance);
-						OutResult.ReplacementObjects.Add(DepotVersion);
-					}
+					OutResult.BaseMaterials.AddUnique(MaterialInstance->GetMaterial());
+					OutResult.ChangelistInstanceCounts.FindOrAdd(FSoftObjectPath(MaterialInstance->GetMaterial()))++;
+					OutResult.ModifiedObjects.Add(MaterialInstance);
 				}
 			}
 
 			if (MaterialFunction != nullptr)
 			{
-				if (bLoadDepotVersions)
+				OutResult.bHasFunctionFiles = true;
+
+				if (bExpandFunctions)
 				{
+					// Note that loading all referencing base materials can be slow for widely-used functions.
 					// Function derived base materials cannot be analyzed in slow mode for two reasons:
 					// * It would probably expand the working set of loaded assets by too much to be useable.
 					// * We would need a way to rebind the temporary depot assets into their referenced materials.
-					OutResult.bHasFunctionFiles = true;
-				}
-				else
-				{
-					// Note that loading all referencing base materials can be slow for widely-used functions.
 					TArray<FAssetData> MaterialAssetDatas;
 					UMaterialEditingLibrary::GetMaterialsReferencingFunction(MaterialFunction, MaterialAssetDatas);
 					for (FAssetData const& AssetData : MaterialAssetDatas)
@@ -681,7 +651,7 @@ private:
 		UMaterialValidationLibrary::GetAllGroups(Groups, /*bInSyncLoad*/true);
 
 		FValidateCollectionResult Collection;
-		CollectFromChangelist(Groups, SelectedEntry->ChangelistState->GetFilesStates(), /*bLoadDepotVersions=*/false, Collection);
+		CollectFromChangelist(Groups, SelectedEntry->ChangelistState->GetFilesStates(), /*bExpandFunctions=*/true, Collection);
 
 		TMap<FSoftObjectPath, FPerMaterialResult> PerMaterialResults;
 		for (UMaterial* BaseMaterial : Collection.BaseMaterials)
@@ -722,25 +692,20 @@ private:
 		FScopedSlowTask SlowTask(2.f, LOCTEXT("SlowValidating", "Running slow validation..."));
 		SlowTask.MakeDialog(/*bShowCancelButton=*/true);
 
-		SlowTask.EnterProgressFrame(1.f, LOCTEXT("SlowFetchHistory", "Fetching revision history..."));
-		{
-			TArray<FString> Filenames;
-			for (FSourceControlStateRef const& FileState : SelectedEntry->ChangelistState->GetFilesStates())
-			{
-				Filenames.Add(FileState->GetFilename());
-			}
-			TSharedRef<FUpdateStatus> UpdateOp = ISourceControlOperation::Create<FUpdateStatus>();
-			UpdateOp->SetUpdateHistory(true);
-			ISourceControlModule::Get().GetProvider().Execute(UpdateOp, Filenames, EConcurrency::Synchronous);
-		}
-
 		SlowTask.EnterProgressFrame(1.f, LOCTEXT("SlowCollecting", "Collecting material assets..."));
 		FValidateCollectionResult Collection;
-		CollectFromChangelist(Groups, SelectedEntry->ChangelistState->GetFilesStates(), /*bLoadDepotVersions=*/true, Collection);
+		CollectFromChangelist(Groups, SelectedEntry->ChangelistState->GetFilesStates(), /*bExpandFunctions=*/false, Collection);
+
+		SlowTask.EnterProgressFrame(1.f, LOCTEXT("SlowFetchHistory", "Fetching revision history..."));
+		FMaterialValidationHelpers::FetchRevisionHistory(Collection.ModifiedObjects);
+		FMaterialValidationHelpers::LoadDepotVersions(Collection.ModifiedObjects, Collection.ReplacementObjects, Collection.DepotPackages);
 
 		if (SlowTask.ShouldCancel())
 		{
-			for (UPackage* DepotPackage : Collection.DepotPackages) { DepotPackage->MarkAsGarbage(); }
+			for (UPackage* DepotPackage : Collection.DepotPackages)
+			{
+				DepotPackage->MarkAsGarbage();
+			}
 			return FReply::Handled();
 		}
 
@@ -752,7 +717,10 @@ private:
 				MaterialTask.EnterProgressFrame(1.f, FText::FromName(BaseMaterial->GetFName()));
 				if (MaterialTask.ShouldCancel())
 				{
-					for (UPackage* DepotPackage : Collection.DepotPackages) { DepotPackage->MarkAsGarbage(); }
+					for (UPackage* DepotPackage : Collection.DepotPackages)
+					{
+						DepotPackage->MarkAsGarbage();
+					}
 					return FReply::Handled();
 				}
 				FSoftObjectPath const Path(BaseMaterial);
diff --git a/Engine/Plugins/Editor/MaterialValidation/Source/Private/MaterialValidators.cpp b/Engine/Plugins/Editor/MaterialValidation/Source/Private/MaterialValidators.cpp
--- a/Engine/Plugins/Editor/MaterialValidation/Source/Private/MaterialValidators.cpp
+++ b/Engine/Plugins/Editor/MaterialValidation/Source/Private/MaterialValidators.cpp
@@ -3,6 +3,11 @@
 #include "MaterialValidators.h"
 
 #include "Editor.h"
+#include "DiffUtils.h"
+#include "ISourceControlModule.h"
+#include "ISourceControlProvider.h"
+#include "ISourceControlRevision.h"
+#include "ISourceControlState.h"
 #include "Logging/TokenizedMessage.h"
 #include "MaterialEditingLibrary.h"
 #include "Materials/Material.h"
@@ -12,6 +17,9 @@
 #include "MaterialValidationGroup.h"
 #include "MaterialValidationLibrary.h"
 #include "Misc/DataValidation.h"
+#include "Misc/PackageName.h"
+#include "SourceControlOperations.h"
+#include "UObject/Package.h"
 
 #include UE_INLINE_GENERATED_CPP_BY_NAME(MaterialValidators)
 
@@ -93,6 +101,8 @@ EDataValidationResult UEditorValidator_MaterialPermutation::ValidateLoadedAsset_
 
 	if (MaterialFunction)
 	{
+		bHasFunctionFiles = true;
+
 		// Add the base materials that use the material function for later valiation.
 		// Note this will require loading UMaterial assets which could be slow for very commonly used material functions.
 		TArray<FAssetData> MaterialAssetDatas;
@@ -130,6 +140,7 @@ void UEditorValidator_MaterialPermutation::PostAssetValidation(TArray<TSharedRef
 		TArray<UMaterialValidationGroup*> Groups;
 		UMaterialValidationLibrary::GetAllGroups(Groups, /*bInSyncLoad*/true);
 
+		// Always do fast validation which compares modified workspace objects against the stored database baseline.
 		for (UMaterial* BaseMaterial : BaseMaterials)
 		{
 			for (UMaterialValidationGroup* Group : Groups)
@@ -145,6 +156,40 @@ void UEditorValidator_MaterialPermutation::PostAssetValidation(TArray<TSharedRef
 				}
 			}
 		}
+
+		// If the fast result already reaches the approval threshold, upgrade to slow validation (depot HEAD vs workspace) for a more accurate delta before emitting the message.
+		// Skip when material functions are involved: their changes cannot be fully captured so the depot comparison would not reflect the true delta.
+		if (ApprovalThreshold > 0 && NumShaders >= ApprovalThreshold && !bHasFunctionFiles && ISourceControlModule::Get().IsEnabled())
+		{
+			FMaterialValidationHelpers::FetchRevisionHistory(MutableView(ModifiedObjects));
+
+			TArray<UMaterialInterface*> ReplacementObjects;
+			TArray<UPackage*> DepotPackages;
+			FMaterialValidationHelpers::LoadDepotVersions(MutableView(ModifiedObjects), ReplacementObjects, DepotPackages);
+
+			int32 SlowNumShaders = 0;
+			for (UMaterial* BaseMaterial : BaseMaterials)
+			{
+				for (UMaterialValidationGroup* Group : Groups)
+				{
+					bool bIsInGroupPath = false;
+					bool bIsInGroup = false;
+					UMaterialValidationLibrary::IsMaterialInGroup(Group, BaseMaterial, bIsInGroupPath, bIsInGroup);
+					if (bIsInGroupPath)
+					{
+						SlowNumShaders -= UMaterialValidationLibrary::GetModifiedShaderCount(Group, BaseMaterial, MutableView(ModifiedObjects), ReplacementObjects, true /*bForceLoadObjects*/);
+						SlowNumShaders += UMaterialValidationLibrary::GetModifiedShaderCount(Group, BaseMaterial, MutableView(ModifiedObjects), {}, true /*bForceLoadObjects*/);
+					}
+				}
+			}
+
+			for (UPackage* DepotPackage : DepotPackages)
+			{
+				DepotPackage->MarkAsGarbage();
+			}
+
+			NumShaders = SlowNumShaders;
+		}
 	}
 
 	if (ShowThreshold > 0 && FMath::Abs(NumShaders) >= ShowThreshold)
@@ -161,6 +206,7 @@ void UEditorValidator_MaterialPermutation::PostAssetValidation(TArray<TSharedRef
 
 	BaseMaterials.Reset();
 	ModifiedObjects.Reset();
+	bHasFunctionFiles = false;
 }
 
 bool UMaterialEditorValidator_Permutation::Validate(UMaterialInterface* InMaterial, FMaterialEditorValidatorContext& InContext) const
@@ -216,4 +262,68 @@ bool UMaterialEditorValidator_Permutation::Validate(UMaterialInterface* InMateri
 	return true;
 }
 
+void FMaterialValidationHelpers::FetchRevisionHistory(TArray<UMaterialInterface*> const& InModifiedObjects)
+{
+	TArray<FString> Filenames;
+	Filenames.Reserve(InModifiedObjects.Num());
+	for (UMaterialInterface const* Object : InModifiedObjects)
+	{
+		if (Object != nullptr)
+		{
+			FString const PackageName = Object->GetOutermost()->GetName();
+			FString Filename;
+			if (FPackageName::TryConvertLongPackageNameToFilename(PackageName, Filename, FPackageName::GetAssetPackageExtension()))
+			{
+				Filenames.AddUnique(Filename);
+			}
+		}
+	}
+
+	if (Filenames.IsEmpty())
+	{
+		return;
+	}
+
+	TSharedRef<FUpdateStatus> UpdateOp = ISourceControlOperation::Create<FUpdateStatus>();
+	UpdateOp->SetUpdateHistory(true);
+	ISourceControlModule::Get().GetProvider().Execute(UpdateOp, Filenames, EConcurrency::Synchronous);
+}
+
+void FMaterialValidationHelpers::LoadDepotVersions(TArray<UMaterialInterface*> const& InModifiedObjects, TArray<UMaterialInterface*>& OutReplacementObjects, TArray<UPackage*>& OutDepotPackages)
+{
+	ISourceControlProvider& Provider = ISourceControlModule::Get().GetProvider();
+
+	OutReplacementObjects.Reset(InModifiedObjects.Num());
+
+	for (UMaterialInterface const* Object : InModifiedObjects)
+	{
+		UMaterialInterface* DepotVersion = nullptr;
+
+		if (Object != nullptr)
+		{
+			FString const PackageName = Object->GetOutermost()->GetName();
+			FString Filename;
+			if (FPackageName::TryConvertLongPackageNameToFilename(PackageName, Filename, FPackageName::GetAssetPackageExtension()))
+			{
+				FSourceControlStatePtr HistoryState = Provider.GetState(Filename, EStateCacheUsage::Use);
+				if (HistoryState.IsValid() && HistoryState->GetHistorySize() > 0)
+				{
+					TSharedPtr<ISourceControlRevision> Revision = HistoryState->GetHistoryItem(0);
+					if (Revision.IsValid())
+					{
+						UPackage* DepotPackage = DiffUtils::LoadPackageForDiff(Revision);
+						if (DepotPackage != nullptr)
+						{
+							OutDepotPackages.Add(DepotPackage);
+							DepotVersion = Cast<UMaterialInterface>(DepotPackage->FindAssetInPackage());
+						}
+					}
+				}
+			}
+		}
+
+		OutReplacementObjects.Add(DepotVersion);
+	}
+}
+
 #undef LOCTEXT_NAMESPACE
diff --git a/Engine/Plugins/Editor/MaterialValidation/Source/Private/MaterialValidators.h b/Engine/Plugins/Editor/MaterialValidation/Source/Private/MaterialValidators.h
--- a/Engine/Plugins/Editor/MaterialValidation/Source/Private/MaterialValidators.h
+++ b/Engine/Plugins/Editor/MaterialValidation/Source/Private/MaterialValidators.h
@@ -7,6 +7,10 @@
 
 #include "MaterialValidators.generated.h"
 
+class UMaterial;
+class UMaterialInterface;
+class UPackage;
+
 /** Implementation of UEditorValidatorBase that applies validation for limiting the addition of material permutations from UMaterials. */
 UCLASS()
 class UEditorValidator_MaterialPermutation : public UEditorValidatorBase
@@ -20,10 +24,12 @@ class UEditorValidator_MaterialPermutation : public UEditorValidatorBase
 private:
 	/** Collected base materials for which to validate the full hierarchy after individual asset validation is complete. */
 	UPROPERTY()
-	TArray<TObjectPtr<class UMaterial>> BaseMaterials;
+	TArray<TObjectPtr<UMaterial>> BaseMaterials;
 	/** Collected individual assets whose changes should be taken into account when doing full hierarchy validation. */
 	UPROPERTY()
-	TArray<TObjectPtr<class UMaterialInterface>> ModifiedObjects;
+	TArray<TObjectPtr<UMaterialInterface>> ModifiedObjects;
+	/** Set when a material function is encountered. We skip the slow validation path because it is more likely to be very slow in this case. */
+	bool bHasFunctionFiles = false;
 };
 
 /** Implementation of UMaterialEditorValidatorBase that applies validation for limiting the addition of material permutations. */
@@ -34,3 +40,13 @@ class UMaterialEditorValidator_Permutation : public UMaterialEditorValidatorBase
 
 	virtual bool Validate(UMaterialInterface* InMaterial, FMaterialEditorValidatorContext& InContext) const override;
 };
+
+/** Helper functions for external validation callers. */
+struct FMaterialValidationHelpers
+{
+	/** Fetches source control revision history for the packages backing the given objects. */
+	static void FetchRevisionHistory(TArray<UMaterialInterface*> const& InModifiedObjects);
+
+	/** Loads the depot HEAD revision for each object. OutReplacementObjects is aligned to InModifiedObjects (nullptr where unavailable). */
+	static void LoadDepotVersions(TArray<UMaterialInterface*> const& InModifiedObjects, TArray<UMaterialInterface*>& OutReplacementObjects, TArray<UPackage*>& OutDepotPackages);
+};
