MaterialValidation: Fix issue where fast validator could over estimate shader count after objects not in the database were loaded.
The correct behavior is to ignore local objects unless they are part of the changelist being validated.
#rnx
#Synced-CL 51704221

diff --git a/Engine/Plugins/Editor/MaterialValidation/Source/Private/MaterialValidationLibrary.cpp b/Engine/Plugins/Editor/MaterialValidation/Source/Private/MaterialValidationLibrary.cpp
--- a/Engine/Plugins/Editor/MaterialValidation/Source/Private/MaterialValidationLibrary.cpp
+++ b/Engine/Plugins/Editor/MaterialValidation/Source/Private/MaterialValidationLibrary.cpp
@@ -773,24 +773,21 @@ static void UpdateValidMaterialInstances(UMaterial const& InMaterial, TConstArra
 /** 
  * Get the asset data associated with a package path.
  * We get this from either the material validation group or from the object directly.
- * We only set InOutMaterialInstance if the object is loaded.
- * For performance we avoid loading an object if it is unloaded (but could optionally do that if that becomes useful).
  * If we return false then it is likely that the material instance exists but:
  * - It wasn't stored in the validation group (it was probably created after last validation group update), and
- * - It isn't in the curent validation changelist, and
- * - It is not currently loaded.
+ * - It isn't in the curent validation changelist.
  */
 static bool GetMaterialInstanceAssetData(
 	FString const& InPackagePath,
-	UMaterialInstanceConstant*& InOutMaterialInstance,
+	UMaterialInstanceConstant const* InMaterialInstance,
 	FMaterialValidationDesc const* InMaterialDesc,
 	FMaterialPermutationData const& InMaterialPermutationData,
 	FMaterialValidationAssetData_MaterialInstance& OutAssetData)
 {
 	// If we were provided with a modified object then use that.
-	if (InOutMaterialInstance != nullptr)
+	if (InMaterialInstance != nullptr)
 	{
-		GetAssetData(*InOutMaterialInstance, InMaterialPermutationData, OutAssetData);
+		GetAssetData(*InMaterialInstance, InMaterialPermutationData, OutAssetData);
 		return true;
 	}
 
@@ -819,16 +816,6 @@ static bool GetMaterialInstanceAssetData(
 		}
 	}
 
-	// If we didn't find in the material validation group but the object is already loaded then use that.
-	FSoftObjectPath SoftObjectPath = MaterialValidationLibrary::ConvertStringPathToSoftObjectPath(InPackagePath);
-	TSoftObjectPtr<UMaterialInstanceConstant> SoftObjectPtr(SoftObjectPath);
-	if (UMaterialInstanceConstant* MaterialInstance = SoftObjectPtr.Get())
-	{
-		GetAssetData(*MaterialInstance, InMaterialPermutationData, OutAssetData);
-		InOutMaterialInstance = MaterialInstance;
-		return true;
-	}
-
 	return false;
 }
 
