Fast path for getting shader compilation parameters

Optimizes material validation shader counting by making a fast path which returns shaders based on translation output but avoids expensive UpdateStaticPermutation() calls that trigger full shader recompilation.
Updates Material Validation Library to use this fast path. Ensures the translation DDC key derived from the dummy built using FMaterialValidationAssetData_MaterialInstance matches the one built from the loaded object, so cached translation results are reused.

- Adds GetShadersUsingCompilationParameters_GameThread() ??? computes shader type info by running the material translator to get FCompilationShaderParameters, rather than relying on conservative estimates or full compilation. When a compiled shader map is already available, uses its compilation output directly.
- Adds UMaterialInstance::SetPermutationParameters() ??? a lightweight method that sets static switch/component mask parameters and base property overrides without creating material resources or triggering shader compilation.
- Extracts GetShaderInfoFromTypeInfo() from UMaterialEditingLibrary::ListShaders() into a reusable free function.
- Updates MaterialValidationLibrary::CreateMaterialInstance() to use the new fast path.
- Updates UMaterialEditingLibrary::ListShaders() to use the new GetShadersUsingCompilationParameters_GameThread().

Note: Materials with layers (bHasMaterialLayers) still fall back to the old UMaterialInstance::GetShaderTypes() path. Further work is needed to support layer overrides through the fast path.

[RN] minor | Rendering Architecture

#Synced-CL 52203036

diff --git a/Engine/Plugins/Editor/MaterialValidation/Source/Private/MaterialValidationLibrary.cpp b/Engine/Plugins/Editor/MaterialValidation/Source/Private/MaterialValidationLibrary.cpp
--- a/Engine/Plugins/Editor/MaterialValidation/Source/Private/MaterialValidationLibrary.cpp
+++ b/Engine/Plugins/Editor/MaterialValidation/Source/Private/MaterialValidationLibrary.cpp
@@ -16,6 +16,8 @@
 #include "MaterialValidationLibraryTypes.h"
 #include "MaterialValidationModule.h"
 #include "Misc/ScopedSlowTask.h"
+#include "MaterialCachedData.h"
+#include "MaterialShared.h"
 
 #include UE_INLINE_GENERATED_CPP_BY_NAME(MaterialValidationLibrary)
 
@@ -819,24 +821,8 @@ static bool GetMaterialInstanceAssetData(
 	return false;
 }
 
-/**
- * Create a transient material instance directly from an AssetData.
- * This can be used to evaluate shader counts for an asset without loading it, and paying the cost of loading references etc.
- * We do need the base material to be loaded, since we use it to create a temporary material instance.
- */
-static UMaterialInstanceConstant* CreateMaterialInstance(UMaterial& InBaseMaterial, FMaterialValidationAssetData_MaterialInstance const& InAssetData)
+static FStaticParameterSet BuildStaticParamSetFromValidationAsset(UMaterial& InBaseMaterial, FMaterialValidationAssetData_MaterialInstance const& InAssetData)
 {
-	// Create a temporary UMaterialInstanceConstant parented to our base material.
-	UMaterialInstanceConstant* MaterialInstance = NewObject<UMaterialInstanceConstant>();
-	MaterialInstance->SetParentEditorOnly(&InBaseMaterial);
-
-	// Apply all of the static properties.
-	FMaterialInstanceBasePropertyOverrides BaseProperties = GetStaticPropertiesAsBasePropertyOverrides(InAssetData.StaticProperties);
-
-	// Calling UpdateStaticPermutation() is slow when we call it 100s of times for a big hierarchy.
-	// Maybe the static switches can't affect the number of shaders? Even if that is true it might not always be true in future...
-	// For now we can take the faster branch below which only sets the base properties (which can modify the shader count).
-#if 0 
 	FStaticParameterSet StaticParameterSet;
 
 	// Fill the static switch parameters.
@@ -867,15 +853,29 @@ static UMaterialInstanceConstant* CreateMaterialInstance(UMaterial& InBaseMateri
 	}
 
 	//todo: We should handle material layer overrides too.
+	//When that happens, FMaterialCachedExpressionData::FunctionInfos should also be copied over for DDC queries to succeed.
+
+	return StaticParameterSet;
+}
 
-	// Set all the static permutation state on the material instance.
-	MaterialInstance->UpdateStaticPermutation(StaticParameterSet, BaseProperties, true);
-#else
-	// Fast path only sets base properties which is still picked up by GetNumShaderTypes()
-	MaterialInstance->BasePropertyOverrides = BaseProperties;
-	MaterialInstance->UpdateOverridableBaseProperties();
+
+/**
+ * Create a transient material instance directly from an AssetData.
+ * This can be used to evaluate shader counts for an asset without loading it, and paying the cost of loading references etc.
+ * We do need the base material to be loaded, since we use it to create a temporary material instance.
+ */
+static UMaterialInstanceConstant* CreateMaterialInstance(UMaterial& InBaseMaterial, FMaterialValidationAssetData_MaterialInstance const& InAssetData)
+{
+	// Create a temporary UMaterialInstanceConstant parented to our base material.
+	UMaterialInstanceConstant* MaterialInstance = NewObject<UMaterialInstanceConstant>();
+	MaterialInstance->SetParentEditorOnly(&InBaseMaterial);
+
+	// Apply all of the static properties.
+	FMaterialInstanceBasePropertyOverrides BaseProperties = GetStaticPropertiesAsBasePropertyOverrides(InAssetData.StaticProperties);
+
+	FStaticParameterSet StaticParamSet = MaterialValidationLibrary::BuildStaticParamSetFromValidationAsset(InBaseMaterial, InAssetData);
+	MaterialInstance->SetPermutationParameters(StaticParamSet, BaseProperties);
 	MaterialInstance->bHasStaticPermutationResource = true;
-#endif
 
 	return MaterialInstance;
 }
@@ -1135,6 +1135,9 @@ int32 UMaterialValidationLibrary::GetModifiedShaderCount(
 	TArray<UMaterialInterface*> const& InReplacementObjects,
 	bool bForceLoadObjects)
 {
+
+	TRACE_CPUPROFILER_EVENT_SCOPE(UMaterialValidationLibrary::GetModifiedShaderCount)
+
 	if (InGroup == nullptr || InBaseMaterial == nullptr)
 	{
 		return 0;
@@ -1226,6 +1229,7 @@ int32 UMaterialValidationLibrary::GetModifiedShaderCount(
 	TArray<uint32> PermutationHashes;
 	PermutationHashes.Add(MaterialAssetData.PermutationHash);
 
+	const bool bBaseHasMaterialLayers = BaseMaterial->GetCachedExpressionData().bHasMaterialLayers;
 	// Walk the material instance hierarchy.
 	for (int32 Index = 1; Index < PackageNames.Num(); ++Index)
 	{
@@ -1268,7 +1272,22 @@ int32 UMaterialValidationLibrary::GetModifiedShaderCount(
 		Child.PermutationHash = MaterialValidation::BuildPermutationHash(Child.StaticProperties, Child.StaticSwitchOverrideValues, Child.ComponentMaskOverrideValues, Child.MaterialLayerHash);
 		if (!PermutationHashes.Contains(Child.PermutationHash))
 		{
-			NumShaders += MaterialValidationLibrary::GetNumShaderTypes(*MaterialValidationLibrary::CreateMaterialInstance(*BaseMaterial, Child));
+			UMaterialInstanceConstant* DummyMI = MaterialValidationLibrary::CreateMaterialInstance(*BaseMaterial, Child);
+
+			// Layers currently do not support the fast path
+			if (bBaseHasMaterialLayers)
+			{
+				TArray<FDebugShaderTypeInfo> ShaderTypeInfos;
+				DummyMI->GetShaderTypes(GMaxRHIShaderPlatform, nullptr, ShaderTypeInfos);
+				const TArray<FDebugShaderInfo> ValidationShaderDebugInfo = UMaterialEditingLibrary::GetShaderInfoFromTypeInfo(ShaderTypeInfos);
+
+				NumShaders += ValidationShaderDebugInfo.Num();
+			}
+			else
+			{
+				NumShaders += MaterialValidationLibrary::GetNumShaderTypes(*DummyMI);
+			}
+
 			PermutationHashes.Add(Child.PermutationHash);
 		}
 	}
diff --git a/Engine/Source/Editor/MaterialEditor/Private/MaterialEditingLibrary.cpp b/Engine/Source/Editor/MaterialEditor/Private/MaterialEditingLibrary.cpp
--- a/Engine/Source/Editor/MaterialEditor/Private/MaterialEditingLibrary.cpp
+++ b/Engine/Source/Editor/MaterialEditor/Private/MaterialEditingLibrary.cpp
@@ -1706,39 +1706,33 @@ int32 UMaterialEditingLibrary::GetNumShaderTypes(UMaterialInterface* Material)
 	return UMaterialEditingLibrary::ListShaders(Material).Num();
 }
 
-TArray<FDebugShaderInfo> UMaterialEditingLibrary::ListShaders(UMaterialInterface* Material)
+TArray<FDebugShaderInfo> UMaterialEditingLibrary::GetShaderInfoFromTypeInfo(const TArray<FDebugShaderTypeInfo>& InShaderTypeInfos)
 {
 	TArray<FDebugShaderInfo> OutShaderInfos;
 
-	if (Material != nullptr)
+	for (FDebugShaderTypeInfo const& ShaderInfo : InShaderTypeInfos)
 	{
-		TArray<FDebugShaderTypeInfo> ShaderInfos;
-		Material->GetShaderTypes(GMaxRHIShaderPlatform, nullptr, ShaderInfos);
+		for (FShaderType* ShaderType : ShaderInfo.ShaderTypes)
+		{
+			if (ShaderType)
+			{
+				OutShaderInfos.Add({
+					ShaderInfo.VFType ? ShaderInfo.VFType->GetFName() : NAME_None,
+					ShaderType->GetFName()
+					});
+			}
+		}
 
-		for (FDebugShaderTypeInfo const& ShaderInfo : ShaderInfos)
+		for (FDebugShaderPipelineInfo const& PipelineInfo : ShaderInfo.Pipelines)
 		{
-			for (FShaderType* ShaderType : ShaderInfo.ShaderTypes)
+			for (FShaderType* ShaderType : PipelineInfo.ShaderTypes)
 			{
 				if (ShaderType)
 				{
 					OutShaderInfos.Add({
 						ShaderInfo.VFType ? ShaderInfo.VFType->GetFName() : NAME_None,
 						ShaderType->GetFName()
-					});
-				}
-			}
-
-			for (FDebugShaderPipelineInfo const& PipelineInfo : ShaderInfo.Pipelines)
-			{
-				for (FShaderType* ShaderType : PipelineInfo.ShaderTypes)
-				{
-					if (ShaderType)
-					{
-						OutShaderInfos.Add({
-							ShaderInfo.VFType ? ShaderInfo.VFType->GetFName() : NAME_None,
-							ShaderType->GetFName()
 						});
-					}
 				}
 			}
 		}
@@ -1747,6 +1741,19 @@ TArray<FDebugShaderInfo> UMaterialEditingLibrary::ListShaders(UMaterialInterface
 	return OutShaderInfos;
 }
 
+TArray<FDebugShaderInfo> UMaterialEditingLibrary::ListShaders(UMaterialInterface* Material)
+{
+	if (!Material)
+	{
+		return {};
+	}
+
+	TArray<FDebugShaderTypeInfo> ValidationShaderTypeInfos = GetShadersUsingCompilationParameters_GameThread(Material);
+	const TArray<FDebugShaderInfo> ValidationShaderDebugInfo = GetShaderInfoFromTypeInfo(ValidationShaderTypeInfos);
+
+	return ValidationShaderDebugInfo;
+}
+
 FMaterialStatistics UMaterialEditingLibrary::GetStatistics(class UMaterialInterface* Material)
 {
 	FMaterialStatistics Result;
diff --git a/Engine/Source/Editor/MaterialEditor/Public/MaterialEditingLibrary.h b/Engine/Source/Editor/MaterialEditor/Public/MaterialEditingLibrary.h
--- a/Engine/Source/Editor/MaterialEditor/Public/MaterialEditingLibrary.h
+++ b/Engine/Source/Editor/MaterialEditor/Public/MaterialEditingLibrary.h
@@ -227,6 +227,15 @@ public:
 	 */
 	static UE_API void RecompileMaterials(TArray<UMaterial*>& Materials, FOnItemComplete const& OnItemComplete);
 
+	/**
+	 * Flattens shader type info into a flat list of individual shader entries,
+	 * stripping the per-type grouping from FDebugShaderTypeInfo.
+	 *
+	 * @param InShaderTypeInfos  Grouped shader type information (e.g., from GetDebugShaderTypeInfo)
+	 * @return Flat array of all individual FDebugShaderInfo across all types
+	 */
+	static UE_API TArray<FDebugShaderInfo> GetShaderInfoFromTypeInfo(const TArray<FDebugShaderTypeInfo>& InShaderTypeInfos);
+
 	/**
 	 *	Layouts the expressions in a grid pattern
 	 */
diff --git a/Engine/Source/Runtime/Engine/Private/Materials/MaterialInstance.cpp b/Engine/Source/Runtime/Engine/Private/Materials/MaterialInstance.cpp
--- a/Engine/Source/Runtime/Engine/Private/Materials/MaterialInstance.cpp
+++ b/Engine/Source/Runtime/Engine/Private/Materials/MaterialInstance.cpp
@@ -4637,6 +4637,22 @@ void UMaterialInstance::UpdateStaticPermutation(const FStaticParameterSet& NewPa
 	}
 }
 
+void UMaterialInstance::SetPermutationParameters(const FStaticParameterSet& InParameters, FMaterialInstanceBasePropertyOverrides& InBasePropertyOverrides)
+{
+	UMaterialInstanceEditorOnlyData* EditorOnly = GetEditorOnlyData();
+	FStaticParameterSet NewParameters = InParameters;
+
+	TrimToOverriddenOnly(NewParameters.StaticSwitchParameters);
+	TrimToOverriddenOnly(NewParameters.EditorOnly.StaticComponentMaskParameters);
+	
+	BasePropertyOverrides = InBasePropertyOverrides;
+
+	UpdateOverridableBaseProperties();
+
+	StaticParametersRuntime = NewParameters.GetRuntime();
+	EditorOnly->StaticParameters = NewParameters.EditorOnly;
+}
+
 TSoftObjectPtr<UObject> UMaterialInstance::GetEnumerationObject(int32 Index) const
 {
 	if (EnumerationObjects.IsValidIndex(Index) && !EnumerationObjects[Index].IsNull())
diff --git a/Engine/Source/Runtime/Engine/Private/Materials/MaterialShared.cpp b/Engine/Source/Runtime/Engine/Private/Materials/MaterialShared.cpp
--- a/Engine/Source/Runtime/Engine/Private/Materials/MaterialShared.cpp
+++ b/Engine/Source/Runtime/Engine/Private/Materials/MaterialShared.cpp
@@ -6882,6 +6882,122 @@ TConstArrayView<TObjectPtr<UTextureCollection>> FEditorMaterialResourceBase::Get
 	return MakeArrayView(ReferencedTextureCollections);
 }
 
+// Helper class for getting translation output from a Material/Material Instance.
+class FDebugShaderInfoMaterialResource : public FMaterialResource
+{
+public:
+	FDebugShaderInfoMaterialResource(UMaterial* InMat, UMaterialInstance* InMatInstance, const ERHIFeatureLevel::Type FeatureLevel, const EShaderPlatform ShaderPlatform,
+		const EMaterialQualityLevel::Type QualityLevel)
+		: FMaterialResource()
+	{
+		SetMaterial(InMat);
+		SetMaterialInstance(InMatInstance);
+
+		SetQualityLevelProperties(ShaderPlatform, QualityLevel);
+	}
+
+	virtual ~FDebugShaderInfoMaterialResource() = default;
+
+	bool TryBuildCompilationOutput(FMaterialCompilationOutput& OutCompilationOutput)
+	{
+		FStaticParameterSet UsedStaticParams;
+
+		if (MaterialInstance)
+		{
+			MaterialInstance->GetStaticParameterValues(UsedStaticParams);
+		}
+		else if (Material)
+		{
+			Material->GetStaticParameterValues(UsedStaticParams);
+		}
+		else
+		{
+			UE_LOGF(LogMaterial, Error, "FDebugShaderInfoMaterialResource created without Material or MaterialInstance");
+			return false;
+		}
+
+		FMaterialShaderMapId ShaderMapId;
+		BuildShaderMapId(ShaderMapId, nullptr);
+
+		// TODO: Handle DDC Key creation once new translator supports it
+		if (ShaderMapId.bUsingNewHLSLGenerator && !bEnableRestrictiveMaterialMode)
+		{
+			FMaterialIRModule Module;
+
+			FMaterialIRModuleBuilder Builder = {
+				.MaterialInterface = GetMaterialInterface(),
+				.ShaderPlatform = GetShaderPlatform(),
+				.TargetPlatform = nullptr,
+				.FeatureLevel = ShaderMapId.FeatureLevel,
+				.QualityLevel = GetQualityLevel(),
+				.BlendMode = GetBlendMode(),
+				.StaticParameters = &UsedStaticParams,
+			};
+
+			if (!Builder.Build(&Module))
+			{
+				return false;
+			}
+
+			OutCompilationOutput = Module.GetCompilationOutput();
+
+			return true;
+		}
+		else
+		{
+			FString MaterialTranslationDDCKeyString = GetMaterialShaderMapKeyString(ShaderMapId, FConservativeShaderParameters(this), GetShaderPlatform(), false);
+			FHLSLMaterialTranslator MaterialTranslator(this, OutCompilationOutput, UsedStaticParams, GetShaderPlatform(), GetQualityLevel(), ShaderMapId.FeatureLevel, nullptr, &ShaderMapId.SubstrateCompilationConfig, MaterialTranslationDDCKeyString);
+			EHLSLMaterialTranslatorResult Result = MaterialTranslator.Translate(false);
+
+			return Result == EHLSLMaterialTranslatorResult::Success;
+		}
+	}
+
+};
+
+TArray<FDebugShaderTypeInfo> GetShadersUsingCompilationParameters_GameThread(UMaterialInterface* InMaterial, const EShaderPlatform InShaderPlatform)
+{
+	UMaterialInstance* Inst = Cast<UMaterialInstance>(InMaterial);
+	FPlatformTypeLayoutParameters LayoutParams;
+	LayoutParams.InitializeForPlatform(nullptr);
+
+	FMaterialResource* Resource = InMaterial->GetMaterialResource(InShaderPlatform);
+
+	if (Resource)
+	{
+		if (Resource->IsGameThreadShaderMapComplete())
+		{
+			const FMaterialShaderMap* ShaderMap = Resource->GetGameThreadShaderMap();
+			const FMaterialCompilationOutput& CompOutput = ShaderMap->GetCompilationOutput();
+
+			TArray<FDebugShaderTypeInfo> ShaderTypeInfos;
+			Resource->GetShaderTypes(LayoutParams, ShaderTypeInfos, FCompilationShaderParameters(Resource, CompOutput));
+
+			return ShaderTypeInfos;
+		}
+
+	}
+	UMaterial* BaseMaterial = InMaterial->GetMaterial();
+	FDebugShaderInfoMaterialResource ShaderInfoResource(BaseMaterial, Inst, GMaxRHIFeatureLevel, InShaderPlatform, EMaterialQualityLevel::Num);
+
+	FMaterialCompilationOutput CompOutput;
+	const bool bValidOutput = ShaderInfoResource.TryBuildCompilationOutput(CompOutput);
+
+	TArray<FDebugShaderTypeInfo> ShaderTypeInfos;
+	if (bValidOutput)
+	{
+		ShaderInfoResource.GetShaderTypes(LayoutParams, ShaderTypeInfos, FCompilationShaderParameters(&ShaderInfoResource, CompOutput));
+	}
+	else
+	{
+		UE_LOGF(LogMaterial, Warning, ("Failed to translate material resource (in %ls hierarchy) for Compilation Output. Returning Conservative Shader Layout"),
+			*ShaderInfoResource.GetFriendlyName());
+
+		ShaderInfoResource.GetShaderTypes(LayoutParams, ShaderTypeInfos, FConservativeShaderParameters(&ShaderInfoResource));
+	}
+
+	return ShaderTypeInfos;
+}
 #endif // WITH_EDITOR
 
 #undef LOCTEXT_NAMESPACE
diff --git a/Engine/Source/Runtime/Engine/Public/MaterialShared.h b/Engine/Source/Runtime/Engine/Public/MaterialShared.h
--- a/Engine/Source/Runtime/Engine/Public/MaterialShared.h
+++ b/Engine/Source/Runtime/Engine/Public/MaterialShared.h
@@ -3688,6 +3688,13 @@ protected:
 	TArray<TObjectPtr<UObject>> ReferencedTextures;
 	TArray<TObjectPtr<UTextureCollection>> ReferencedTextureCollections;
 };
+
+/**
+ * Utility method for getting the shaders for the current platform, using FCompilationShaderParameters.
+ * In case the ShaderMap is available, it uses the Compilation output within to generate the shaders.
+ * In case it is not, it handles creating that output efficiently, without triggering shader compilation.
+ */
+ENGINE_API TArray<FDebugShaderTypeInfo> GetShadersUsingCompilationParameters_GameThread(UMaterialInterface* InMaterial, const EShaderPlatform InShaderPlatform = GMaxRHIShaderPlatform);
 #endif
 
 inline bool ShouldIncludeMaterialInDefaultOpaquePass(const FMaterial& Material)
diff --git a/Engine/Source/Runtime/Engine/Public/Materials/MaterialInstance.h b/Engine/Source/Runtime/Engine/Public/Materials/MaterialInstance.h
--- a/Engine/Source/Runtime/Engine/Public/Materials/MaterialInstance.h
+++ b/Engine/Source/Runtime/Engine/Public/Materials/MaterialInstance.h
@@ -1050,10 +1050,16 @@ public:
 	 * Can also trigger recompile based on new set of FMaterialInstanceBasePropertyOverrides 
 	 */
 	ENGINE_API void UpdateStaticPermutation(const FStaticParameterSet& NewParameters, FMaterialInstanceBasePropertyOverrides& NewBasePropertyOverrides, const bool bForceStaticPermutationUpdate = false, FMaterialUpdateContext* MaterialUpdateContext = nullptr);
-
 	/** Updates all default textures of this material instance, based on permutations and Parent Material. */
 	ENGINE_API void UpdateReferencedDefaultTextures();
 
+	/**
+	* Set just the permutation parameters, without triggering any shader compilation.
+	* This is useful for operations which require the Material Instance to have the right data
+	* but it will be invalid for rendering until UpdateStaticPermutation() is called.
+	*/
+	ENGINE_API void SetPermutationParameters(const FStaticParameterSet& InParameters, FMaterialInstanceBasePropertyOverrides& InBasePropertyOverrides);
+
 	/**
 	* Sets new static parameter overrides on the instance and recompiles the static permutation resources if needed.
 	* Can be passed either a minimal parameter set (overridden parameters only) or the entire set generated by GetStaticParameterValues().
